Megascans plugin integration.

diff --git a/Engine/Plugins/MegascansPlugin/Config/FilterPlugin.ini b/Engine/Plugins/MegascansPlugin/Config/FilterPlugin.ini
new file mode 100644
--- /dev/null
+++ b/Engine/Plugins/MegascansPlugin/Config/FilterPlugin.ini
@@ -0,0 +1,8 @@
+[FilterPlugin]
+; This section lists additional files which will be packaged along with your plugin. Paths should be listed relative to the root plugin directory, and
+; may include "...", "*", and "?" wildcards to match directories, files, and individual characters respectively.
+;
+; Examples:
+;    /README.txt
+;    /Extras/...
+;    /Binaries/ThirdParty/*.dll
diff --git a/Engine/Plugins/MegascansPlugin/MegascansPlugin.uplugin b/Engine/Plugins/MegascansPlugin/MegascansPlugin.uplugin
new file mode 100644
--- /dev/null
+++ b/Engine/Plugins/MegascansPlugin/MegascansPlugin.uplugin
@@ -0,0 +1,51 @@
+{
+	"FileVersion": 3,
+	"Version": 1,
+	"VersionName": "0.10",
+	"FriendlyName": "Megascans",
+	"Description": "Megascans Link for Quixel Bridge.",
+	"Category": "Editor",
+	"CreatedBy": "Quixel",
+	"CreatedByURL": "http://www.quixel.se",
+	"DocsURL": "https://help.quixel.com/hc/en-us/sections/360000575958-UE4-Livelink-FAQs",
+	"MarketplaceURL": "",
+	"SupportURL": "[REDACTED_EMAIL]",	
+	"CanContainContent": true,
+	"EnabledByDefault" : true,
+	"SupportedTargetPlatforms": [
+		"Win64",
+		"Mac",
+		"Linux"
+	],
+	"Modules": [
+		{
+			"Name": "MegascansPlugin",
+			"Type": "Editor",
+			"LoadingPhase": "PostDefault",
+			"WhitelistPlatforms": [
+				"Win64",
+				"Mac",
+				"Linux"
+			]
+		}
+	],
+	"Plugins": [
+		{
+			"Name": "PythonScriptPlugin",
+			"Enabled": true
+		},
+		{
+			"Name": "EditorScriptingUtilities",
+			"Enabled": true
+		},
+		{
+			"Name": "AlembicImporter",
+			"Enabled": true
+		},
+		{
+			"Name": "SQLiteSupport",
+			"Enabled": true
+		}
+
+	]
+}
\ No newline at end of file
diff --git a/Engine/Plugins/MegascansPlugin/Source/MegascansPlugin/MegascansPlugin.Build.cs b/Engine/Plugins/MegascansPlugin/Source/MegascansPlugin/MegascansPlugin.Build.cs
new file mode 100644
--- /dev/null
+++ b/Engine/Plugins/MegascansPlugin/Source/MegascansPlugin/MegascansPlugin.Build.cs
@@ -0,0 +1,44 @@
+// Copyright Epic Games, Inc. All Rights Reserved.
+namespace UnrealBuildTool.Rules
+{
+	public class MegascansPlugin : ModuleRules
+	{
+		public MegascansPlugin(ReadOnlyTargetRules Target)
+			: base(Target)
+		{
+			PCHUsage = PCHUsageMode.UseExplicitOrSharedPCHs;
+
+			PrivateDependencyModuleNames.AddRange(
+				new string[]
+				{
+					"AssetRegistry",
+					"ContentBrowser",
+					"Core",
+					"CoreUObject",
+					"EditorStyle",
+					"Engine",
+					"LevelEditor",
+					"PythonScriptPlugin",
+					"Settings",
+					"Slate",
+					"SlateCore",
+                    "Sockets",
+                    "Networking",
+                    "EditorScriptingUtilities",
+					"Json",
+                    "UnrealEd",
+                    "MaterialEditor",
+					"AlembicImporter",
+					"AlembicLibrary",
+					"FoliageEdit",
+                    "Foliage",
+					"HTTP",
+					"SQLiteCore",
+					"SQLiteSupport",
+					"StaticMeshEditor",
+
+				}
+			);
+		}
+	}
+}
\ No newline at end of file
diff --git a/Engine/Plugins/MegascansPlugin/Source/MegascansPlugin/Private/AssetImportData.cpp b/Engine/Plugins/MegascansPlugin/Source/MegascansPlugin/Private/AssetImportData.cpp
new file mode 100644
--- /dev/null
+++ b/Engine/Plugins/MegascansPlugin/Source/MegascansPlugin/Private/AssetImportData.cpp
@@ -0,0 +1,4 @@
+// Copyright Epic Games, Inc. All Rights Reserved.
+#include "AssetImportData.h"
+
+DEFINE_LOG_CATEGORY(MSLiveLinkLog);
\ No newline at end of file
diff --git a/Engine/Plugins/MegascansPlugin/Source/MegascansPlugin/Private/AssetImportData.h b/Engine/Plugins/MegascansPlugin/Source/MegascansPlugin/Private/AssetImportData.h
new file mode 100644
--- /dev/null
+++ b/Engine/Plugins/MegascansPlugin/Source/MegascansPlugin/Private/AssetImportData.h
@@ -0,0 +1,94 @@
+// Copyright Epic Games, Inc. All Rights Reserved.
+#pragma once
+#include "CoreMinimal.h"
+
+DECLARE_LOG_CATEGORY_EXTERN(MSLiveLinkLog, Log, All);
+
+struct FAssetMetaData {
+	FString Resolution;
+	FString Category;
+	FString Type;
+	FString Id;
+	FString Name;
+	FString Path;
+	FString TextureFormat;
+	TArray<FString> Tags;
+	FString ActiveLOD;
+	TArray<FString> Categories;
+	FString ExportPath;
+	FString NamingConvention;
+	FString FolderNamingConvention;
+	FString MinLOD;
+	bool bUseBillboardMaterial;
+	TMap<FString, TArray<int8>> MaterialTypes;
+	bool bIsModularWindow;
+	bool bSavePackages;
+};
+
+
+
+struct FAssetTextureData {
+	FString Format;
+	FString Type;
+	FString Resolution;
+	FString Name;
+	FString NameOverride;
+	FString Path;
+};
+
+struct FAssetBillboardData
+{
+	FString Type;	
+	FString Path;
+};
+
+struct FAssetPackedTextures {
+	TSharedPtr<FAssetTextureData> PackedTextureData;
+	TMap<FString, TArray<FString>> ChannelData;
+};
+
+struct FAssetMeshData {
+	FString Format;
+	FString Type;
+	FString Resolution;
+	FString Name;
+	FString NameOverride;
+	FString Path;
+};
+
+
+
+struct FAssetLodData {
+	FString Lod;
+	FString Path;
+	FString Name;
+	FString NameOverride;
+	FString LodObjectName;
+	FString Format;
+	FString Type;
+};
+
+
+
+struct FAssetTypeData {
+	TSharedPtr<FAssetMetaData> AssetMetaInfo;
+	TArray<TSharedPtr<FAssetTextureData>> TextureComponents;
+	TArray<TSharedPtr<FAssetMeshData>> MeshList;
+	TArray<TSharedPtr<FAssetLodData>> LodList;
+	TArray<TSharedPtr<FAssetPackedTextures>> PackedTextures;
+
+	TArray<TSharedPtr<FAssetBillboardData>> BillboardTextures;
+	TMap<FString, TMap<FString, float>> PlantsLodScreenSizes;
+};
+
+
+struct FAssetsData {
+	TArray<TSharedPtr<FAssetTypeData>> AllAssetsData;
+};
+
+struct FDHIData {
+	FString CharacterPath;
+	FString CommonPath;
+	FString RootPath;
+	FString CharacterName;
+};
diff --git a/Engine/Plugins/MegascansPlugin/Source/MegascansPlugin/Private/AssetImportDataHandler.cpp b/Engine/Plugins/MegascansPlugin/Source/MegascansPlugin/Private/AssetImportDataHandler.cpp
new file mode 100644
--- /dev/null
+++ b/Engine/Plugins/MegascansPlugin/Source/MegascansPlugin/Private/AssetImportDataHandler.cpp
@@ -0,0 +1,307 @@
+// Copyright Epic Games, Inc. All Rights Reserved.
+#include "AssetImportDataHandler.h"
+
+TSharedPtr<FAssetDataHandler> FAssetDataHandler::AssetDataHandlerInst;
+
+FAssetDataHandler::FAssetDataHandler()
+{
+	
+}
+
+TSharedPtr<FAssetDataHandler> FAssetDataHandler::Get()
+{
+	if (!AssetDataHandlerInst.IsValid())
+	{
+		AssetDataHandlerInst = MakeShareable(new FAssetDataHandler);
+
+	}
+	return AssetDataHandlerInst;
+}
+
+TSharedPtr<FAssetsData> FAssetDataHandler::GetAssetsData(const FString& AssetsImportJson)
+{
+	FString ImportData = ConcatJsonString(AssetsImportJson);
+	AssetsImportData = MakeShareable(new FAssetsData);	
+	TSharedPtr<FJsonObject> ImportDataObject = DeserializeJson(ImportData);	
+	TArray<TSharedPtr<FJsonValue> > AssetsImportDataArray = ImportDataObject->GetArrayField(TEXT("Assets"));
+	for (TSharedPtr<FJsonValue> AssetDataObject : AssetsImportDataArray)
+	{
+		TSharedPtr<FAssetTypeData> ParsedAssetData = GetAssetData(AssetDataObject->AsObject());
+		AssetsImportData->AllAssetsData.Add(ParsedAssetData);
+	}	
+	return AssetsImportData;
+}
+
+
+
+TSharedPtr<FAssetTypeData> FAssetDataHandler::GetAssetData(TSharedPtr<FJsonObject> AssetDataObject)
+{
+	TSharedPtr<FAssetTypeData> ParsedAssetData = MakeShareable(new FAssetTypeData);
+	ParsedAssetData->AssetMetaInfo = GetAssetMetaData(AssetDataObject);	
+
+	TArray<TSharedPtr<FJsonValue>> TextureDataArray = AssetDataObject->GetArrayField(TEXT("components"));	
+
+	for (TSharedPtr<FJsonValue> TextureData : TextureDataArray)
+	{
+		ParsedAssetData->TextureComponents.Add(GetAssetTextureData(TextureData->AsObject()));
+	}	
+
+	TArray<TSharedPtr<FJsonValue>> MeshDataArray = AssetDataObject->GetArrayField(TEXT("meshList"));
+	for (TSharedPtr<FJsonValue> MeshData : MeshDataArray)
+	{
+		ParsedAssetData->MeshList.Add(GetAssetMeshData(MeshData->AsObject()));
+	}
+
+	TArray<TSharedPtr<FJsonValue>> LodDataArray = AssetDataObject->GetArrayField(TEXT("lodList"));	
+	for (TSharedPtr<FJsonValue> LodData : LodDataArray)
+	{
+		ParsedAssetData->LodList.Add(GetAssetLodData(LodData->AsObject()));
+	}
+
+	TArray<TSharedPtr<FJsonValue>> PackedDataArray = AssetDataObject->GetArrayField(TEXT("packedTextures"));
+	for (TSharedPtr<FJsonValue> PackedData : PackedDataArray)
+	{
+		ParsedAssetData->PackedTextures.Add(GetPackedTextureData(PackedData->AsObject()));
+	}
+
+	if (ParsedAssetData->AssetMetaInfo->Type == TEXT("3dplant"))
+	{
+		TArray<TSharedPtr<FJsonValue>> BillboardDataArray = AssetDataObject->GetArrayField(TEXT("components-billboard"));
+		if (BillboardDataArray.Num() > 0)
+		{
+			for (TSharedPtr<FJsonValue> BillboardData : BillboardDataArray)
+			{
+				ParsedAssetData->BillboardTextures.Add(GetBillboardData(BillboardData->AsObject()));
+			}
+		}
+
+		//Meta tags array for use in Lod screen sizes
+		TArray<TSharedPtr<FJsonValue>> MetaTagsArray = AssetDataObject->GetArrayField(TEXT("meta"));
+		ParsedAssetData->PlantsLodScreenSizes = GetLodScreenSizes(MetaTagsArray);	
+
+		//Get the Use Billboard tag.
+		ParsedAssetData->AssetMetaInfo->bUseBillboardMaterial = GetBillboardMaterialSubtype(MetaTagsArray);
+
+	}
+	//Get the material ids for modular windows from json
+	if (ParsedAssetData->AssetMetaInfo->Type == TEXT("3d"))
+	{
+		GetMulitpleMaterialIds(AssetDataObject->GetArrayField(TEXT("meta")), ParsedAssetData->AssetMetaInfo);
+	}
+
+	return ParsedAssetData;	
+}
+
+TSharedPtr<FAssetMetaData> FAssetDataHandler::GetAssetMetaData(TSharedPtr<FJsonObject> MetaDataObject)
+{
+	TSharedPtr<FAssetMetaData> ParsedMetaData = MakeShareable(new FAssetMetaData);
+
+	ParsedMetaData->Category = MetaDataObject->GetStringField(TEXT("category"));
+	ParsedMetaData->Type = MetaDataObject->GetStringField(TEXT("type"));
+	ParsedMetaData->Id = MetaDataObject->GetStringField(TEXT("id"));
+	ParsedMetaData->Name = MetaDataObject->GetStringField(TEXT("name"));
+	ParsedMetaData->Path = MetaDataObject->GetStringField(TEXT("path"));
+	ParsedMetaData->TextureFormat = MetaDataObject->GetStringField(TEXT("textureFormat"));
+	ParsedMetaData->ActiveLOD = MetaDataObject->GetStringField(TEXT("activeLOD"));
+	
+	ParsedMetaData->ExportPath = MetaDataObject->GetStringField(TEXT("exportPath"));
+	ParsedMetaData->NamingConvention = MetaDataObject->GetStringField(TEXT("namingConvention"));
+	ParsedMetaData->FolderNamingConvention = MetaDataObject->GetStringField(TEXT("folderNamingConvention"));
+	ParsedMetaData->Resolution = MetaDataObject->GetStringField(TEXT("resolution"));
+
+	if (MetaDataObject->HasField(TEXT("minLOD")))
+	{
+		ParsedMetaData->MinLOD = MetaDataObject->GetStringField(TEXT("minLOD"));
+	}
+	else {
+		ParsedMetaData->MinLOD = TEXT("");
+	}
+	
+
+	TArray<TSharedPtr<FJsonValue>> TagsArray = MetaDataObject->GetArrayField(TEXT("tags"));
+	for (TSharedPtr<FJsonValue> TagData : TagsArray)
+	{
+		ParsedMetaData->Tags.Add(TagData->AsString());
+	}
+
+	TArray<TSharedPtr<FJsonValue>> CategoriesArray = MetaDataObject->GetArrayField(TEXT("categories"));
+	for (TSharedPtr<FJsonValue> CategoryData : CategoriesArray)
+	{
+		ParsedMetaData->Categories.Add(CategoryData->AsString());
+	}
+
+	return ParsedMetaData;
+	
+}
+
+TSharedPtr< FAssetPackedTextures> FAssetDataHandler::GetPackedTextureData(TSharedPtr<FJsonObject> PackedDataObject)
+{
+	TArray<FString> ChannelKeys = { "Red", "Green","Blue","Alpha","Grayscale" };
+	TSharedPtr<FAssetPackedTextures> ParsedPackedData = MakeShareable(new FAssetPackedTextures);
+	ParsedPackedData->PackedTextureData = GetAssetTextureData(PackedDataObject);
+	TSharedPtr<FJsonObject> ChannelsData = PackedDataObject->GetObjectField(TEXT("channelsData"));
+
+	for (FString ChannelKey : ChannelKeys)
+	{
+		
+		TArray<FString> ChannelValues;
+		if (!ChannelsData->HasField(ChannelKey)) continue;
+		TArray<TSharedPtr<FJsonValue>> ChannelKeyData = ChannelsData->GetArrayField(ChannelKey);
+		for (TSharedPtr<FJsonValue> ChData : ChannelKeyData)
+		{
+			
+			ChannelValues.Add(ChData->AsString());			
+		}
+		ParsedPackedData->ChannelData.Add(ChannelKey, ChannelValues);
+	}
+
+	return ParsedPackedData;
+}
+
+TSharedPtr<FAssetTextureData> FAssetDataHandler::GetAssetTextureData(TSharedPtr<FJsonObject> TextureDataObject)
+{
+	
+	TSharedPtr<FAssetTextureData> ParsedTextureData = MakeShareable(new FAssetTextureData);
+	ParsedTextureData->Format = TextureDataObject->GetStringField(TEXT("format"));
+	ParsedTextureData->Type = TextureDataObject->GetStringField(TEXT("type"));
+	ParsedTextureData->Resolution = TextureDataObject->GetStringField(TEXT("resolution"));
+	ParsedTextureData->Name = GetAssetName(TextureDataObject->GetStringField(TEXT("name")));
+	ParsedTextureData->NameOverride = TextureDataObject->GetStringField(TEXT("nameOverride"));
+	ParsedTextureData->Path = TextureDataObject->GetStringField(TEXT("path"));
+
+	return ParsedTextureData;
+}
+
+
+
+TSharedPtr<FAssetMeshData> FAssetDataHandler::GetAssetMeshData(TSharedPtr<FJsonObject> MeshDataObject)
+{
+	TSharedPtr<FAssetMeshData> ParsedMeshData = MakeShareable(new FAssetMeshData);
+	ParsedMeshData->Format = MeshDataObject->GetStringField(TEXT("format"));
+	ParsedMeshData->Type = MeshDataObject->GetStringField(TEXT("type"));
+	ParsedMeshData->Resolution = MeshDataObject->GetStringField(TEXT("resolution"));
+	ParsedMeshData->Name = GetAssetName(MeshDataObject->GetStringField(TEXT("name")));
+	ParsedMeshData->NameOverride = MeshDataObject->GetStringField(TEXT("nameOverride"));
+	ParsedMeshData->Path = MeshDataObject->GetStringField(TEXT("path"));
+
+	return ParsedMeshData;
+}
+
+TSharedPtr<FAssetLodData> FAssetDataHandler::GetAssetLodData(TSharedPtr<FJsonObject> LodDataObject)
+{
+	TSharedPtr<FAssetLodData> ParsedLodData = MakeShareable(new FAssetLodData);
+	ParsedLodData->Lod = LodDataObject->GetStringField(TEXT("lod"));
+	ParsedLodData->Path = LodDataObject->GetStringField(TEXT("path"));
+	ParsedLodData->Name = LodDataObject->GetStringField(TEXT("name"));
+	ParsedLodData->NameOverride = LodDataObject->GetStringField(TEXT("nameOverride"));
+	ParsedLodData->LodObjectName = LodDataObject->GetStringField(TEXT("lodObjectName"));
+	ParsedLodData->Format = LodDataObject->GetStringField(TEXT("format"));
+	ParsedLodData->Type = LodDataObject->GetStringField(TEXT("type"));
+
+	return ParsedLodData;
+}
+
+TSharedPtr<FAssetBillboardData> FAssetDataHandler::GetBillboardData(TSharedPtr<FJsonObject> BillboardObject)
+{
+	TSharedPtr<FAssetBillboardData> ParsedBillboardData = MakeShareable(new FAssetBillboardData);
+	ParsedBillboardData->Path = BillboardObject->GetStringField(TEXT("path"));
+	ParsedBillboardData->Type = BillboardObject->GetStringField(TEXT("type"));
+	return ParsedBillboardData;
+}
+
+TMap<FString, TMap<FString, float>> FAssetDataHandler::GetLodScreenSizes(TArray<TSharedPtr<FJsonValue>> MetaTagsArray)
+{
+	TMap<FString, TMap<FString, float>> LODScreenSizes;
+
+	for (TSharedPtr<FJsonValue> MetaData : MetaTagsArray)
+	{
+		TSharedPtr<FJsonObject> MetaObject = MetaData->AsObject();
+		
+		if (MetaObject->GetStringField(TEXT("key")) == TEXT("lodDistance"))
+		{
+			TArray<TSharedPtr<FJsonValue>> ScreenSizesArray = MetaObject->GetArrayField(TEXT("value"));
+			for (TSharedPtr<FJsonValue> ScreenSizes : ScreenSizesArray)
+			{
+				
+				FString VariationKey = FString::FromInt(ScreenSizes->AsObject()->GetIntegerField(TEXT("variation")));
+				VariationKey = TEXT("Var") + VariationKey;
+				TArray<TSharedPtr<FJsonValue>> VariationSSArray = ScreenSizes->AsObject()->GetArrayField(TEXT("distance"));
+				
+				TMap<FString, float> VariationScreenData;				
+				for (TSharedPtr<FJsonValue> VarScreenSize : VariationSSArray)
+				{					
+					FString Lod = TEXT("lod") + FString::FromInt(VarScreenSize->AsObject()->GetIntegerField(TEXT("lod")));
+					float ScreenSize = VarScreenSize->AsObject()->GetNumberField(TEXT("lodDistance"));
+					
+					VariationScreenData.Add( Lod, ScreenSize);
+				}
+				
+				LODScreenSizes.Add(VariationKey, VariationScreenData);
+			}
+
+			break;
+		}
+	}
+
+	return LODScreenSizes;
+}
+
+bool FAssetDataHandler::GetBillboardMaterialSubtype(TArray<TSharedPtr<FJsonValue>> MetaTagsArray)
+{
+	for (TSharedPtr<FJsonValue> MetaData : MetaTagsArray)
+	{
+		TSharedPtr<FJsonObject> MetaObject = MetaData->AsObject();
+
+		if (MetaObject->GetStringField(TEXT("key")) == TEXT("useBillboardMaterial"))
+		{
+			return MetaObject->GetBoolField(TEXT("value"));
+			break;
+		}
+	}
+	return false;
+}
+
+void FAssetDataHandler::GetMulitpleMaterialIds(TArray<TSharedPtr<FJsonValue>> MetaTagsArray, TSharedPtr<FAssetMetaData> AssetMetaInfo)
+{
+	for (TSharedPtr<FJsonValue> MetaData : MetaTagsArray)
+	{
+		TSharedPtr<FJsonObject> MetaObject = MetaData->AsObject();
+
+		if (MetaObject->GetStringField(TEXT("key")) == TEXT("materialIds"))
+		{
+			TArray<TSharedPtr<FJsonValue>> MaterialIds = MetaObject->GetArrayField(TEXT("value"));
+			for (TSharedPtr<FJsonValue> MasterialIDData : MaterialIds)
+			{
+				FString MaterialType = MasterialIDData->AsObject()->GetStringField(TEXT("material"));
+				TArray<TSharedPtr<FJsonValue>> Ids = MasterialIDData->AsObject()->GetArrayField(TEXT("ids"));
+				TArray<int8> MeshMaterialIds;
+				for (TSharedPtr<FJsonValue> MeshMaterialId : Ids)
+				{
+					MeshMaterialIds.Add(MeshMaterialId->AsNumber());
+				}
+
+				AssetMetaInfo->MaterialTypes.Add(MaterialType, MeshMaterialIds);
+			}
+			AssetMetaInfo->bIsModularWindow = true;
+			return;
+			
+		}
+	}
+	AssetMetaInfo->bIsModularWindow = false;
+}
+
+
+FString FAssetDataHandler::GetAssetName(const FString& AssetFileName)
+{
+	FString AssetName, AssetExtension;	
+	AssetFileName.Split(TEXT("."), &AssetName, &AssetExtension);
+	return AssetName;
+}
+
+FString FAssetDataHandler::ConcatJsonString(const FString& AssetsImportJson)
+{
+	FString StartString = TEXT("{\"Assets\":");
+	FString EndString = TEXT("}");
+	FString FinalString = StartString + AssetsImportJson + EndString;	
+	return FinalString;
+}
\ No newline at end of file
diff --git a/Engine/Plugins/MegascansPlugin/Source/MegascansPlugin/Private/AssetImportDataHandler.h b/Engine/Plugins/MegascansPlugin/Source/MegascansPlugin/Private/AssetImportDataHandler.h
new file mode 100644
--- /dev/null
+++ b/Engine/Plugins/MegascansPlugin/Source/MegascansPlugin/Private/AssetImportDataHandler.h
@@ -0,0 +1,30 @@
+// Copyright Epic Games, Inc. All Rights Reserved.
+#pragma once
+#include "CoreMinimal.h"
+#include "AssetImportData.h"
+#include "Utilities/MiscUtils.h"
+
+class FAssetDataHandler {
+private:
+	FAssetDataHandler();
+	TSharedPtr<FAssetsData> AssetsImportData ;
+	
+
+	FString GetAssetName(const FString & AssetName);
+	FString ConcatJsonString(const FString& AssetsImportJson);
+	static TSharedPtr<FAssetDataHandler> AssetDataHandlerInst;
+	TSharedPtr<FAssetTypeData> GetAssetData(TSharedPtr<FJsonObject> AssetDataObject);
+	TSharedPtr< FAssetPackedTextures> GetPackedTextureData(TSharedPtr<FJsonObject> PackedDataObject);
+	TSharedPtr<FAssetMetaData> GetAssetMetaData(TSharedPtr<FJsonObject> MetaDataObject);
+	TSharedPtr<FAssetTextureData> GetAssetTextureData(TSharedPtr<FJsonObject> TextureDataObject);
+	TSharedPtr<FAssetMeshData> GetAssetMeshData(TSharedPtr<FJsonObject> MeshDataObject);
+	TSharedPtr<FAssetLodData> GetAssetLodData(TSharedPtr<FJsonObject> LodDataObject);
+	TSharedPtr<FAssetBillboardData> GetBillboardData(TSharedPtr<FJsonObject> BillboardObject);
+	TMap<FString, TMap<FString, float>> GetLodScreenSizes(TArray<TSharedPtr<FJsonValue>> MetaTagsArray);
+	bool GetBillboardMaterialSubtype(TArray<TSharedPtr<FJsonValue>> MetaTagsArray);
+	void GetMulitpleMaterialIds(TArray<TSharedPtr<FJsonValue>> MetaTagsArray, TSharedPtr<FAssetMetaData> AssetMetaInfo);
+
+public:	
+	static TSharedPtr<FAssetDataHandler> Get();
+	TSharedPtr<FAssetsData> GetAssetsData(const FString& AssetsImportJson);
+};
\ No newline at end of file
diff --git a/Engine/Plugins/MegascansPlugin/Source/MegascansPlugin/Private/AssetImporters/Import3D.cpp b/Engine/Plugins/MegascansPlugin/Source/MegascansPlugin/Private/AssetImporters/Import3D.cpp
new file mode 100644
--- /dev/null
+++ b/Engine/Plugins/MegascansPlugin/Source/MegascansPlugin/Private/AssetImporters/Import3D.cpp
@@ -0,0 +1,263 @@
+// Copyright Epic Games, Inc. All Rights Reserved.
+#include "Import3D.h"
+#include "Utilities/AssetsDatabase.h"
+#include "AssetImporters/ImportSurface.h"
+
+
+#include "Runtime/Engine/Classes/Engine/StaticMesh.h"
+#include "Utilities/MeshOp.h"
+#include "Utilities/MiscUtils.h"
+#include "Utilities/MaterialUtils.h"
+
+#include "EditorAssetLibrary.h"
+#include "EditorStaticMeshLibrary.h"
+#include "UI/MSSettings.h"
+#include "Runtime/Core/Public/Misc/MessageDialog.h"
+#include "Runtime/Core/Public/Internationalization/Text.h"
+
+#include "StaticMeshEditorSubsystem.h"
+
+TSharedPtr<FImport3d> FImport3d::Import3dInst;
+
+TSharedPtr<FImport3d> FImport3d::Get()
+{
+	if (!Import3dInst.IsValid())
+	{
+		Import3dInst = MakeShareable(new FImport3d);
+	}
+	return Import3dInst;
+}
+
+
+void FImport3d::ImportAsset(TSharedPtr<FAssetTypeData> AssetImportData)
+{	
+	const UMegascansSettings* MegascansSettings = GetDefault<UMegascansSettings>();
+	TSharedPtr<FAssetImportParams> AssetSetupParameters = FAssetImportParams::Get();
+	TSharedPtr<ImportParams3DAsset> Asset3DParameters = AssetSetupParameters->Get3DAssetsParams(AssetImportData);	
+
+	if (AssetImportData->MeshList.Num() == 0) return;
+
+	if (AssetImportData->AssetMetaInfo->ActiveLOD == TEXT("high"))
+	{
+		EAppReturnType::Type ContinueImport = FMessageDialog::Open(EAppMsgType::OkCancel, FText(FText::FromString("You are about to import a high poly mesh. This may cause Unreal to stop responding. Press Ok to continue.")));
+		if (ContinueImport == EAppReturnType::Cancel) return;
+	}
+	UMaterialInstanceConstant* MaterialInstance = FImportSurface::Get()->ImportSurface(AssetImportData, Asset3DParameters->ParamsAssetType->MaterialParams[0]);
+
+
+	(Asset3DParameters->ParamsAssetType->SubType == EAsset3DSubtype::MULTI_MESH) ? ImportScatter(AssetImportData, MaterialInstance, Asset3DParameters) : ImportNormal(AssetImportData, MaterialInstance, Asset3DParameters);
+
+	AssetUtils::FocusOnSelected(Asset3DParameters->BaseParams->AssetDestination);
+	AssetRecord MSRecord;
+	MSRecord.ID = AssetImportData->AssetMetaInfo->Id;
+	MSRecord.Name = Asset3DParameters->BaseParams->AssetName;
+	MSRecord.Path = Asset3DParameters->BaseParams->AssetDestination;
+	MSRecord.Type = AssetImportData->AssetMetaInfo->Type;
+	FAssetsDatabase::Get()->AddRecord(MSRecord);
+
+
+}
+
+
+void FImport3d::ImportScatter(TSharedPtr<FAssetTypeData> AssetImportData, UMaterialInstanceConstant* MaterialInstance, TSharedPtr<ImportParams3DAsset> Asset3DParameters)
+{
+	TSharedPtr<FMeshOps> MeshUtils = FMeshOps::Get();
+
+	const UMegascansSettings* MegascansSettings = GetDefault<UMegascansSettings>();
+	TArray<FString> ExistingAssets = GetAssetsList(Asset3DParameters->BaseParams->AssetDestination);
+	MeshUtils->ImportMesh(AssetImportData->MeshList[0]->Path, Asset3DParameters->ParamsAssetType->MeshDestination, "");
+	TArray<FString> ImportedAssets;
+	TArray<FString> AllAssets = GetAssetsList(Asset3DParameters->ParamsAssetType->MeshDestination);
+
+	TMap<FString, TArray<UStaticMesh*>> ScatterLods;
+
+
+
+	for (FString Asset : AllAssets)
+	{
+		if (!ExistingAssets.Contains(Asset))
+		{
+			ImportedAssets.Add(Asset);
+			UStaticMesh* ImportedMesh = CastChecked<UStaticMesh>(UEditorAssetLibrary::LoadAsset(Asset));
+			if (MaterialInstance != nullptr)
+				ImportedMesh->SetMaterial(0, CastChecked<UMaterialInterface>(MaterialInstance));
+
+			ImportedMesh->PostEditChange();
+		}
+	}
+
+	if (MegascansSettings->bEnableLods && AssetImportData->LodList.Num() > 0) {
+
+		TArray<FString> LodPathList = ParseLodList(AssetImportData);
+		FString LodDestination = FPaths::Combine(Asset3DParameters->ParamsAssetType->MeshDestination, TEXT("Lods"));
+		int32 LodCounter = 1;
+		UStaticMeshEditorSubsystem* SMESubsystem = GEditor->GetEditorSubsystem<UStaticMeshEditorSubsystem>();
+		for (FString LodPath : LodPathList)
+		{
+			if (LodCounter > 7) continue;
+
+			MeshUtils->ImportMesh(LodPath, LodDestination);
+			TArray<FString> ImportedLods = UEditorAssetLibrary::ListAssets(LodDestination);
+
+			for (FString MeshLod : ImportedLods)
+			{
+				FString MeshLodName = FPaths::GetBaseFilename(MeshLod);
+
+				TArray<FString> NameTags;
+				FString LodNameRefined = TEXT("");
+
+				MeshLodName.ParseIntoArray(NameTags, TEXT("_"));
+				for (FString Tag : NameTags)
+				{
+					if (!Tag.Contains(TEXT("LOD")))
+					{
+						LodNameRefined.Append(Tag);
+					}
+				}
+
+
+
+				for (FString SourceMeshPath : ImportedAssets)
+				{
+					FString SourceMeshName = FPaths::GetBaseFilename(SourceMeshPath);
+					TArray<FString> MeshNameTags;
+					FString MeshNameRefined = TEXT("");
+					UStaticMesh* ImportedScatter = CastChecked<UStaticMesh>(UEditorAssetLibrary::LoadAsset(SourceMeshPath));
+
+					SourceMeshName.ParseIntoArray(MeshNameTags, TEXT("_"));
+
+					for (FString Tag : MeshNameTags)
+					{
+						if (!Tag.Contains(TEXT("LOD")))
+						{
+							MeshNameRefined.Append(Tag);
+						}
+					}
+
+					if (MeshNameRefined == LodNameRefined)
+					{
+						SMESubsystem->SetLodFromStaticMesh(ImportedScatter, LodCounter, CastChecked<UStaticMesh>(UEditorAssetLibrary::LoadAsset(MeshLod)), 0, true);
+						//UEditorStaticMeshLibrary::SetLodFromStaticMesh(ImportedScatter, LodCounter, CastChecked<UStaticMesh>(UEditorAssetLibrary::LoadAsset(MeshLod)), 0, true);
+
+					}
+					ImportedScatter->PostEditChange();
+					ImportedScatter->MarkPackageDirty();
+
+				}
+
+
+			}
+			LodCounter++;
+			UEditorAssetLibrary::DeleteDirectory(LodDestination);
+		}
+	}
+
+	int32 VarCounter = 1;
+	FString OverridenName;
+	FString OverridenExtension;
+	AssetImportData->MeshList[0]->NameOverride.Split(TEXT("."), &OverridenName, &OverridenExtension);
+
+	OverridenName = SanitizeName(OverridenName);
+	for (FString Asset : ImportedAssets)
+	{
+		FString BasePath = FPaths::GetPath(Asset);
+		FString RenamedPath = FPaths::Combine(BasePath, (OverridenName + TEXT("_") + FString::FromInt(VarCounter)));
+
+		if (MaterialInstance != nullptr)
+			CastChecked<UStaticMesh>(UEditorAssetLibrary::LoadAsset(Asset))->SetMaterial(1, CastChecked<UMaterialInterface>(MaterialInstance));
+
+		UEditorAssetLibrary::RenameAsset(Asset, RenamedPath);
+
+
+		UStaticMesh* ImportedScatter = CastChecked<UStaticMesh>(UEditorAssetLibrary::LoadAsset(RenamedPath));
+		MeshUtils->RemoveExtraMaterialSlot(ImportedScatter);
+
+		if (MegascansSettings->bCreateFoliage)
+		{
+
+			MeshUtils->CreateFoliageAsset(Asset3DParameters->ParamsAssetType->MeshDestination, ImportedScatter, (OverridenName + TEXT("_") + FString::FromInt(VarCounter)));
+		}
+		ImportedScatter->MarkPackageDirty();
+		if (AssetImportData->AssetMetaInfo->bSavePackages)
+		{
+			AssetUtils::SavePackage(ImportedScatter);
+		}
+		//AssetUtils::SavePackage(ImportedScatter);
+		VarCounter++;
+	}
+
+	MeshUtils.Reset();
+
+}
+
+void FImport3d::ImportNormal(TSharedPtr<FAssetTypeData> AssetImportData, UMaterialInstanceConstant* MaterialInstance, TSharedPtr<ImportParams3DAsset> Asset3DParameters)
+{
+	const UMegascansSettings* MegascansSettings = GetDefault<UMegascansSettings>();
+	FString FileExtension;
+	FString FilePath;
+	AssetImportData->MeshList[0]->Path.Split(TEXT("."), &FilePath, &FileExtension);
+	FileExtension = FPaths::GetExtension(AssetImportData->MeshList[0]->Path);
+
+	FString OverridenName;
+	FString OverridenExtension;
+	AssetImportData->MeshList[0]->NameOverride.Split(TEXT("."), &OverridenName, &OverridenExtension);
+
+	TSharedPtr<FMeshOps> MeshUtils = FMeshOps::Get();
+	FString AssetPath = MeshUtils->ImportMesh(AssetImportData->MeshList[0]->Path, Asset3DParameters->ParamsAssetType->MeshDestination, SanitizeName(OverridenName));
+
+	if (!UEditorAssetLibrary::DoesAssetExist(AssetPath)) return;
+	UStaticMesh* ImportedAsset = CastChecked<UStaticMesh>(UEditorAssetLibrary::LoadAsset(AssetPath));
+	if (ImportedAsset == nullptr) return;
+	if (MaterialInstance != nullptr)
+		ImportedAsset->SetMaterial(0, CastChecked<UMaterialInterface>(MaterialInstance));
+
+	if (MegascansSettings->bEnableLods && AssetImportData->LodList.Num() > 0)
+	{
+		TArray<FString> LodPathList = ParseLodList(AssetImportData);
+		if (FileExtension == TEXT("abc"))
+		{
+
+			MeshUtils->ApplyAbcLods(ImportedAsset, LodPathList, Asset3DParameters->ParamsAssetType->MeshDestination);
+		}
+		else {
+
+			MeshUtils->ApplyLods(LodPathList, ImportedAsset);
+		}
+		if (!AssetImportData->AssetMetaInfo->bIsModularWindow)
+		{
+			MeshUtils->RemoveExtraMaterialSlot(ImportedAsset);
+		}
+
+		else {
+			FString GlassMaterialPath = GetMaterial(TEXT("GlassMasterMaterial"));
+			UMaterialInstanceConstant* GlassInstance = FMaterialUtils::CreateInstanceMaterial(GlassMaterialPath, Asset3DParameters->BaseParams->AssetDestination, TEXT("GlassMaterial_inst"));
+			for (const TPair<FString, TArray<int8>> MaterialInfo : AssetImportData->AssetMetaInfo->MaterialTypes)
+			{
+				if (MaterialInfo.Key == TEXT("glass"))
+				{
+					for (int8 MatId : MaterialInfo.Value)
+					{
+						ImportedAsset->SetMaterial(MatId - 1, CastChecked<UMaterialInterface>(GlassInstance));
+					}
+				}
+				else
+				{
+					for (int8 MatId : MaterialInfo.Value)
+					{
+						ImportedAsset->SetMaterial(MatId - 1, CastChecked<UMaterialInterface>(MaterialInstance));
+					}
+
+				}
+			}
+
+		}
+	}
+	ImportedAsset->MarkPackageDirty();
+	ImportedAsset->PostEditChange();
+	if (AssetImportData->AssetMetaInfo->bSavePackages)
+	{
+		AssetUtils::SavePackage(ImportedAsset);
+	}
+	//AssetUtils::SavePackage(ImportedAsset);	
+	MeshUtils.Reset();
+}
diff --git a/Engine/Plugins/MegascansPlugin/Source/MegascansPlugin/Private/AssetImporters/Import3D.h b/Engine/Plugins/MegascansPlugin/Source/MegascansPlugin/Private/AssetImporters/Import3D.h
new file mode 100644
--- /dev/null
+++ b/Engine/Plugins/MegascansPlugin/Source/MegascansPlugin/Private/AssetImporters/Import3D.h
@@ -0,0 +1,25 @@
+// Copyright Epic Games, Inc. All Rights Reserved.
+#pragma once
+#include "CoreMinimal.h"
+
+#include "AssetImportData.h"
+#include "ImportFactory.h"
+#include "Utilities/AssetData.h"
+
+class UMaterialInstanceConstant;
+
+class FImport3d : public IImportAsset
+{
+private:
+	FImport3d() = default;
+	static TSharedPtr<FImport3d> Import3dInst;
+	void ImportScatter(TSharedPtr<FAssetTypeData> AssetImportData, UMaterialInstanceConstant* MaterialInstance, const FString& MeshDestination, const FString& AssetName);
+
+	void ImportScatter(TSharedPtr<FAssetTypeData> AssetImportData, UMaterialInstanceConstant* MaterialInstance, TSharedPtr<ImportParams3DAsset> Asset3DParameters);
+	void ImportNormal(TSharedPtr<FAssetTypeData> AssetImportData, UMaterialInstanceConstant* MaterialInstance, TSharedPtr<ImportParams3DAsset> Asset3DParameters);
+public:
+	static TSharedPtr<FImport3d> Get();
+	
+	virtual void ImportAsset(TSharedPtr<FAssetTypeData> AssetImportData) override;
+
+};
\ No newline at end of file
diff --git a/Engine/Plugins/MegascansPlugin/Source/MegascansPlugin/Private/AssetImporters/Import3DPlant.cpp b/Engine/Plugins/MegascansPlugin/Source/MegascansPlugin/Private/AssetImporters/Import3DPlant.cpp
new file mode 100644
--- /dev/null
+++ b/Engine/Plugins/MegascansPlugin/Source/MegascansPlugin/Private/AssetImporters/Import3DPlant.cpp
@@ -0,0 +1,316 @@
+// Copyright Epic Games, Inc. All Rights Reserved.
+#include "Import3DPlant.h"
+
+
+#include "Utilities/AssetsDatabase.h"
+#include "Utilities/MeshOp.h"
+#include "AssetImportData.h"
+
+#include "Runtime/Engine/Classes/Engine/StaticMesh.h"
+#include "EditorAssetLibrary.h"
+#include "EditorStaticMeshLibrary.h"
+
+
+#include "Utilities/MiscUtils.h"
+#include "UI/MSSettings.h"
+#include "PerPlatformProperties.h"
+
+TSharedPtr<FImportPlant> FImportPlant::ImportPlantInst;
+
+
+
+TSharedPtr<FImportPlant> FImportPlant::Get()
+{
+	if (!ImportPlantInst.IsValid())
+	{
+		ImportPlantInst = MakeShareable(new FImportPlant);
+	}
+	return ImportPlantInst;
+}
+
+void FImportPlant::ImportAsset(TSharedPtr<FAssetTypeData> AssetImportData)
+{
+	if (AssetImportData->MeshList.Num() == 0) return;
+	const UMegascansSettings* MegascansSettings = GetDefault<UMegascansSettings>();
+	TSharedPtr<FAssetImportParams> AssetSetupParameters = FAssetImportParams::Get();
+	TSharedPtr<ImportParams3DPlantAsset> AssetPlantParameters = AssetSetupParameters->Get3DPlantParams(AssetImportData);	
+	PlantImportType ImportType = GetImportType(AssetImportData);
+	TMap<FString, FString> ImportedPlants = ImportPlants(AssetImportData, AssetPlantParameters);
+	FString LastLOD = AssetImportData->LodList[AssetImportData->LodList.Num() - 1]->Lod;
+	FString MinLOD = (AssetImportData->AssetMetaInfo->MinLOD == TEXT("")) ? LastLOD : AssetImportData->AssetMetaInfo->MinLOD;
+	int32 HighestLodIndex = FCString::Atoi(*MinLOD.Replace(TEXT("lod"), TEXT("")));
+	HighestLodIndex = HighestLodIndex - FCString::Atoi(*AssetImportData->AssetMetaInfo->ActiveLOD.Replace(TEXT("lod"), TEXT("")));
+
+	UMaterialInstanceConstant* BillboardMatInst = nullptr;
+	UMaterialInstanceConstant* MaterialInstance = nullptr;
+
+	if (ImportType == PlantImportType::BLLBOARD_NORMAL && !MegascansSettings->bEnableLods)
+	{
+		ImportType = PlantImportType::NORMAL_ONLY;
+	}
+
+	switch (ImportType)
+	{
+	case PlantImportType::BILLBOARD_ONLY:
+		AssetImportData->TextureComponents.Empty();
+		PopulateBillboardTextures(AssetImportData);	
+		BillboardMatInst = FImportSurface::Get()->ImportSurface(AssetImportData, AssetPlantParameters->ParamsAssetType->BillboardMasterMaterial);
+		
+		ApplyMaterial(ImportType, BillboardMatInst, nullptr, ImportedPlants, AssetImportData->AssetMetaInfo->bSavePackages);
+		break;
+
+	case PlantImportType::BLLBOARD_NORMAL:
+		
+		MaterialInstance = FImportSurface::Get()->ImportSurface( AssetImportData, AssetPlantParameters->ParamsAssetType->PlantsMasterMaterial);
+		AssetImportData->TextureComponents.Empty();
+		PopulateBillboardTextures(AssetImportData);		
+		BillboardMatInst = FImportSurface::Get()->ImportSurface(AssetImportData, AssetPlantParameters->ParamsAssetType->BillboardMasterMaterial);
+		ApplyMaterial(ImportType, MaterialInstance, BillboardMatInst, ImportedPlants,  AssetImportData->AssetMetaInfo->bSavePackages);
+		break;
+
+	case PlantImportType::NORMAL_ONLY:
+		MaterialInstance = FImportSurface::Get()->ImportSurface(AssetImportData, AssetPlantParameters->ParamsAssetType->PlantsMasterMaterial);
+		ApplyMaterial(ImportType, MaterialInstance, nullptr, ImportedPlants, AssetImportData->AssetMetaInfo->bSavePackages);
+		break;
+
+	default:
+		break;
+	}
+
+	AssetUtils::FocusOnSelected(AssetPlantParameters->BaseParams->AssetDestination);
+	AssetRecord MSRecord;
+	MSRecord.ID = AssetImportData->AssetMetaInfo->Id;
+	MSRecord.Name = AssetPlantParameters->BaseParams->AssetName;
+	MSRecord.Path = AssetPlantParameters->BaseParams->AssetDestination;
+	MSRecord.Type = AssetImportData->AssetMetaInfo->Type;
+	FAssetsDatabase::Get()->AddRecord(MSRecord);
+}
+
+
+
+PlantImportType FImportPlant::GetImportType(TSharedPtr<FAssetTypeData> AssetImportData)
+{	
+	FString LastLOD = AssetImportData->LodList[AssetImportData->LodList.Num() - 1]->Lod;
+	FString MinLOD = (AssetImportData->AssetMetaInfo->MinLOD == TEXT("")) ? LastLOD : AssetImportData->AssetMetaInfo->MinLOD;
+	FString CurrentLOD = (AssetImportData->AssetMetaInfo->ActiveLOD == TEXT("high")) ? TEXT("lod-1") : AssetImportData->AssetMetaInfo->ActiveLOD;	
+
+	if (MinLOD == CurrentLOD && AssetImportData->BillboardTextures.Num() != 0)
+	{		
+		return PlantImportType::BILLBOARD_ONLY;
+	}
+	else if (LastLOD != MinLOD)
+	{
+		return PlantImportType::NORMAL_ONLY;
+	}
+	else {
+		if(AssetImportData->BillboardTextures.Num() != 0) return PlantImportType::BLLBOARD_NORMAL;
+		else return PlantImportType::NORMAL_ONLY;
+	}
+}
+
+
+
+
+void FImportPlant::PopulateBillboardTextures(TSharedPtr<FAssetTypeData> AssetImportData)
+{
+	for (auto BillboardComponent : AssetImportData->BillboardTextures)
+	{
+		TSharedPtr<FAssetTextureData> TextureComponent = MakeShareable(new FAssetTextureData);
+		TextureComponent->Name = FPaths::GetCleanFilename(BillboardComponent->Path);
+		TextureComponent->NameOverride = TextureComponent->Name;
+		TextureComponent->Path = BillboardComponent->Path;
+		TextureComponent->Format = FPaths::GetExtension(BillboardComponent->Path);
+		TextureComponent->Type = BillboardComponent->Type;
+		AssetImportData->TextureComponents.Add(TextureComponent);
+	}
+
+}
+
+void FImportPlant::SetBillboardImportParams(TSharedPtr<SurfaceImportParams> SImportParams)
+{
+	FString BillboardsFolder = TEXT("Billboards");
+	SImportParams->MInstanceDestination = FPaths::Combine(SImportParams->MInstanceDestination, BillboardsFolder);
+	SImportParams->TexturesDestination = FPaths::Combine(SImportParams->TexturesDestination, BillboardsFolder);
+	SImportParams->MInstanceName = SImportParams->AssetName + TEXT("_billboard_inst");
+}
+
+TMap<FString, FString> FImportPlant::ImportPlants(TSharedPtr<FAssetTypeData> AssetImportData, TSharedPtr<ImportParams3DPlantAsset> AssetPlantParameters)
+{
+	const UMegascansSettings* MegascansSettings = GetDefault<UMegascansSettings>();
+	TArray<FString> AllLodList = ParseLodList(AssetImportData);
+
+	TArray<TSharedPtr<FAssetLodData>> SelectedLods = ParsePlantsLodList(AssetImportData);
+
+	TMap<FString, FString> ImportedPlantVariations;
+
+	FString FileExtension;
+	FString FilePath;
+	AssetImportData->MeshList[0]->Path.Split(TEXT("."), &FilePath, &FileExtension);
+	FileExtension = FPaths::GetExtension(AssetImportData->MeshList[0]->Path);
+
+	int32 VarCounter = 1;
+	
+	for (auto PlantVar : AssetImportData->MeshList)
+	{	
+
+		FString VariantPath = FMeshOps::Get()->ImportMesh(PlantVar->Path, AssetPlantParameters->ParamsAssetType->MeshDestination);
+		if (!UEditorAssetLibrary::DoesAssetExist(VariantPath)) continue;
+		UStaticMesh* ImportedAsset = CastChecked<UStaticMesh>(LoadAsset(VariantPath));
+		FString VariantName = FPaths::GetBaseFilename(VariantPath);
+		FString Variant, LodNumber;
+		VariantName.Split(TEXT("_"), &Variant, &LodNumber);
+		TArray<FString> VarLodList;
+		if (MegascansSettings->bEnableLods && AssetImportData->LodList.Num() > 0)
+		{
+			// Alternate implementation for TArray based variable, gives more control
+			//TMap<FString, TSharedPtr<FAssetLodData>> VariationLodsImported;			
+			TArray<FString> VariationLodsImported;
+			VariationLodsImported.Add(AssetImportData->AssetMetaInfo->ActiveLOD);
+			
+			for(TSharedPtr<FAssetLodData> LodData : SelectedLods)
+			{
+				FString LodBaseFilename = FPaths::GetBaseFilename(LodData->Path);
+
+				FString VariantTag;
+				FString LodTag;
+
+				LodBaseFilename.Split(TEXT("_"), &VariantTag, &LodTag);
+				if (VariantTag == Variant)
+				{
+					VariationLodsImported.Add(LodData->Lod);
+					//VariationLodsImported.Add(LodData->Lod, LodData);					
+					VarLodList.Add(LodData->Path);
+				}
+			}
+
+
+			if (FileExtension == "abc")
+			{
+				FMeshOps::Get()->ApplyAbcLods(ImportedAsset, VarLodList, AssetPlantParameters->ParamsAssetType->MeshDestination);
+			}
+			else {
+
+				FMeshOps::Get()->ApplyLods(VarLodList, ImportedAsset);
+			}			
+			//Apply the custom lod screen sizes.
+			if (AssetImportData->PlantsLodScreenSizes.Contains(Variant))
+			{
+				SetLodScreenSizes(ImportedAsset, AssetImportData->PlantsLodScreenSizes[Variant], VariationLodsImported);
+			}
+		}
+		
+		FString OverridenName;
+		FString OverridenExtension;
+		PlantVar->NameOverride.Split(TEXT("."), &OverridenName, &OverridenExtension);
+		OverridenName = SanitizeName(OverridenName);
+		FString PlantAssetName = OverridenName + TEXT("_") + FString::FromInt(VarCounter);
+		VarCounter++;
+		if (MegascansSettings->bCreateFoliage)
+		{
+			FMeshOps::Get()->CreateFoliageAsset(AssetPlantParameters->ParamsAssetType->MeshDestination, ImportedAsset, PlantAssetName);
+		}
+		ImportedAsset->MarkPackageDirty();
+		
+		FString RenamedAssetPath = FPaths::Combine(FPaths::GetPath(VariantPath), PlantAssetName);
+		//ImportedPlants.Add(RenamedAssetPath);
+		FString Variation;
+		FString Name;
+		PlantVar->Name.Split(TEXT("_"), &Variation, &Name);
+		ImportedPlantVariations.Add(Variation, RenamedAssetPath);
+		UEditorAssetLibrary::RenameAsset(VariantPath, RenamedAssetPath);
+
+		if (AssetImportData->AssetMetaInfo->bSavePackages)
+		{
+			AssetUtils::SavePackage(ImportedAsset);
+		}
+	}
+	return ImportedPlantVariations;
+}
+
+void FImportPlant::ApplyMaterial(PlantImportType ImportType, UMaterialInstanceConstant* MaterialInstance, UMaterialInstanceConstant* BillboardInstance, TMap<FString, FString> ImportedPlants,  bool bSavePackage)
+{	
+	
+	for (const TPair<FString, FString>& PlantVar : ImportedPlants  )
+	{
+		UStaticMesh* ImportedPlantMesh = CastChecked<UStaticMesh>(UEditorAssetLibrary::LoadAsset(PlantVar.Value));
+		if(MaterialInstance !=nullptr)
+			ImportedPlantMesh->SetMaterial(0, CastChecked<UMaterialInterface>(MaterialInstance));
+
+		if (ImportType == PlantImportType::BLLBOARD_NORMAL && BillboardInstance != nullptr)
+		{
+			
+			AssetUtils::AddStaticMaterial(ImportedPlantMesh, BillboardInstance);			
+			FMeshSectionInfo MeshSectionInfo = ImportedPlantMesh->GetSectionInfoMap().Get(ImportedPlantMesh->GetNumLODs()-1, 0);
+			MeshSectionInfo.MaterialIndex = ImportedPlantMesh->StaticMaterials.Num() - 1;			
+			ImportedPlantMesh->GetSectionInfoMap().Set(ImportedPlantMesh->GetNumLODs() - 1, 0, MeshSectionInfo);
+			ImportedPlantMesh->Modify();
+			ImportedPlantMesh->PostEditChange();
+			ImportedPlantMesh->MarkPackageDirty();
+
+
+		}
+
+		if (bSavePackage)
+		{
+			AssetUtils::SavePackage(ImportedPlantMesh);
+		}
+
+	}
+
+}
+
+void FImportPlant::SetLodScreenSizes(TSharedPtr<FAssetTypeData> AssetImportData, TMap<FString, FString> ImportedPlants)
+{
+	
+
+}
+
+void FImportPlant::SetLodScreenSizes(UStaticMesh* SourceMesh, TMap<FString, float>& LodScreenSizes, const TArray<FString>& SelectedLods)
+{
+	//Check if lod screen size information exist in the json.
+	if (LodScreenSizes.Contains(TEXT("lod0")) && LodScreenSizes[TEXT("lod0")] == 0) return;
+
+	//Check if any lods are skipped.
+	TArray<int8> ImportedLods;
+	int8 ActiveLod = FCString::Atoi(*SelectedLods[0].Replace(TEXT("lod"), TEXT("")));
+	bool bSequence = true;
+
+	for (int8 i = 1 ; i < SelectedLods.Num()-1; i++)
+	{
+		int8 CurrentLod = FCString::Atoi(*SelectedLods[i].Replace(TEXT("lod"), TEXT("")));
+		if (CurrentLod != ActiveLod + 1) {
+			bSequence = false;
+			break;
+		}
+		ActiveLod = CurrentLod;
+	}
+
+
+	
+	if (bSequence)
+	{
+		if (SelectedLods[0] != TEXT("lod0") && LodScreenSizes.Contains(SelectedLods[0]))
+		{
+			LodScreenSizes.Remove(SelectedLods[0]);
+			LodScreenSizes.Add(SelectedLods[0], 1.0);			
+		}
+		int8 count = 0;
+		SourceMesh->bAutoComputeLODScreenSize = false;
+		for (FString Lod : SelectedLods)
+		{
+			FStaticMeshSourceModel& SourceModel = SourceMesh->GetSourceModel(count);
+			if (LodScreenSizes.Contains(Lod)) {
+				SourceModel.ScreenSize = FPerPlatformFloat(LodScreenSizes[Lod]);
+			}
+			else {
+				SourceMesh->bAutoComputeLODScreenSize = true;
+				return;
+			}
+
+			count++;
+		}
+	}
+
+
+	
+}
diff --git a/Engine/Plugins/MegascansPlugin/Source/MegascansPlugin/Private/AssetImporters/Import3DPlant.h b/Engine/Plugins/MegascansPlugin/Source/MegascansPlugin/Private/AssetImporters/Import3DPlant.h
new file mode 100644
--- /dev/null
+++ b/Engine/Plugins/MegascansPlugin/Source/MegascansPlugin/Private/AssetImporters/Import3DPlant.h
@@ -0,0 +1,40 @@
+// Copyright Epic Games, Inc. All Rights Reserved.
+#pragma once
+#include "CoreMinimal.h"
+
+#include "AssetImporters/ImportSurface.h"
+#include "Utilities/AssetData.h"
+
+
+class UMaterialInstanceConstant;
+
+enum PlantImportType {
+	BILLBOARD_ONLY,
+	NORMAL_ONLY,
+	BLLBOARD_NORMAL
+};
+
+enum PlantPrefsType
+{
+	NORMALPrefs,
+	BILLBOARDPrefs
+};
+
+class FImportPlant : public  IImportAsset
+{
+private:
+	FImportPlant() = default;
+	static TSharedPtr<FImportPlant> ImportPlantInst;	
+	PlantImportType GetImportType(TSharedPtr<FAssetTypeData> AssetImportData);
+	
+	void PopulateBillboardTextures(TSharedPtr<FAssetTypeData> AssetImportData);
+	void SetBillboardImportParams(TSharedPtr<SurfaceImportParams> SImportParams);
+	TMap<FString, FString> ImportPlants(TSharedPtr<FAssetTypeData> AssetImportData, TSharedPtr<ImportParams3DPlantAsset> AssetPlantParameters);
+	void ApplyMaterial(PlantImportType ImportType, UMaterialInstanceConstant* MaterialInstance, UMaterialInstanceConstant* BillboardInstance, TMap<FString, FString> ImportedPlants, bool bSavePackage = false);
+	void SetLodScreenSizes(TSharedPtr<FAssetTypeData> AssetImportData, TMap<FString, FString> ImportedPlants);
+	void SetLodScreenSizes(UStaticMesh* SourceMesh, TMap<FString, float>& LodScreenSizes, const TArray<FString>& SelectedLods);
+public:
+	static TSharedPtr<FImportPlant> Get();
+	virtual void ImportAsset(TSharedPtr<FAssetTypeData> AssetImportData) override;
+	
+};
\ No newline at end of file
diff --git a/Engine/Plugins/MegascansPlugin/Source/MegascansPlugin/Private/AssetImporters/ImportFactory.cpp b/Engine/Plugins/MegascansPlugin/Source/MegascansPlugin/Private/AssetImporters/ImportFactory.cpp
new file mode 100644
--- /dev/null
+++ b/Engine/Plugins/MegascansPlugin/Source/MegascansPlugin/Private/AssetImporters/ImportFactory.cpp
@@ -0,0 +1,36 @@
+// Copyright Epic Games, Inc. All Rights Reserved.
+#include "ImportFactory.h"
+#include "AssetImporters/ImportSurface.h"
+
+
+TSharedPtr<F3DImportFactory> F3DImportFactory::Import3dInst;
+TMap<FString, TSharedPtr<IImportAsset>> FAssetImportFactory::RegisteredImporters;
+bool FAssetImportFactory::bImportersRegistered;
+
+TSharedPtr<F3DImportFactory> F3DImportFactory::Get()
+{
+	if (!Import3dInst.IsValid())
+	{
+		Import3dInst = MakeShareable(new F3DImportFactory);
+	}
+	return Import3dInst;
+}
+
+bool FAssetImportFactory::RegisterImporters()
+{
+	//if (!bImportersRegistered) return true;
+
+	RegisteredImporters[TEXT("3d")] = nullptr;
+	//RegisteredImporters[TEXT("surface")] = FImportSurface::Get();
+	return false;
+}
+
+void FAssetImportFactory::Import(TSharedPtr<FAssetTypeData> AssetImportData)
+{
+	RegisteredImporters[AssetImportData->AssetMetaInfo->Type]->ImportAsset(AssetImportData);
+}
+
+void F3DImportFactory::ImportAsset(TSharedPtr<FAssetTypeData> AssetImportData)
+{
+
+}
diff --git a/Engine/Plugins/MegascansPlugin/Source/MegascansPlugin/Private/AssetImporters/ImportFactory.h b/Engine/Plugins/MegascansPlugin/Source/MegascansPlugin/Private/AssetImporters/ImportFactory.h
new file mode 100644
--- /dev/null
+++ b/Engine/Plugins/MegascansPlugin/Source/MegascansPlugin/Private/AssetImporters/ImportFactory.h
@@ -0,0 +1,49 @@
+// Copyright Epic Games, Inc. All Rights Reserved.
+#pragma once
+#include "CoreMinimal.h"
+#include "AssetImportData.h"
+
+
+
+class IImportAsset
+{
+public:
+	IImportAsset() = default;
+	virtual ~IImportAsset() = default;
+	virtual void ImportAsset(TSharedPtr<FAssetTypeData> AssetImportData) = 0;
+};
+
+class FAssetImportFactory
+{
+public:	
+	FAssetImportFactory() = delete;
+	static bool RegisterImporters();
+	static void Import(TSharedPtr<FAssetTypeData> AssetImportData);
+private :	
+	static TMap<FString, TSharedPtr<IImportAsset>> RegisteredImporters;
+	static bool bImportersRegistered ;
+
+};
+
+
+class F3DImportFactory : public IImportAsset
+{
+private:
+	F3DImportFactory() = default;
+	static TSharedPtr<F3DImportFactory> Import3dInst;
+	virtual void ImportAsset(TSharedPtr<FAssetTypeData> AssetImportData);	
+
+public:
+	static TSharedPtr<F3DImportFactory> Get();
+	
+
+};
+
+
+
+
+
+
+
+
+
diff --git a/Engine/Plugins/MegascansPlugin/Source/MegascansPlugin/Private/AssetImporters/ImportSurface.cpp b/Engine/Plugins/MegascansPlugin/Source/MegascansPlugin/Private/AssetImporters/ImportSurface.cpp
new file mode 100644
--- /dev/null
+++ b/Engine/Plugins/MegascansPlugin/Source/MegascansPlugin/Private/AssetImporters/ImportSurface.cpp
@@ -0,0 +1,433 @@
+// Copyright Epic Games, Inc. All Rights Reserved.
+#include "ImportSurface.h"
+
+#include "Utilities/AssetsDatabase.h"
+#include "Utilities/MiscUtils.h"
+#include "AssetToolsModule.h"
+#include "IAssetTools.h"
+#include "AssetImportTask.h"
+#include "Editor/UnrealEd/Classes/Factories/MaterialInstanceConstantFactoryNew.h"
+#include "Runtime/Engine/Classes/Materials/MaterialInstanceConstant.h"
+#include "Runtime/Engine/Classes/Engine/Texture.h"
+#include "UI/MSSettings.h"
+#include "MaterialEditingLibrary.h"
+#include "ContentBrowserModule.h"
+#include "IContentBrowserSingleton.h"
+
+#include "EditorAssetLibrary.h"
+
+#include "Runtime/Engine/Classes/Engine/Selection.h"
+#include "Editor.h"
+#include "PackageTools.h"
+
+TSharedPtr<FImportSurface> FImportSurface::ImportSurfaceInst;
+
+void FImportSurface::ImportAsset(TSharedPtr<FAssetTypeData> AssetImportData)
+{
+	TSharedPtr<FAssetImportParams> AssetSetupParameters = FAssetImportParams::Get();	
+	TSharedPtr<ImportParamsSurfaceAsset> SurfaceImportParams = (AssetImportData->AssetMetaInfo->Type == TEXT("surface")) ? AssetSetupParameters->GetSurfaceParams(AssetImportData) : AssetSetupParameters->GetAtlasBrushParams(AssetImportData);
+	UMaterialInstanceConstant* MaterialInstance = ImportSurface(AssetImportData, SurfaceImportParams->ParamsAssetType);
+	AssetUtils::FocusOnSelected(SurfaceImportParams->ParamsAssetType->MaterialInstanceDestination);
+	if (AssetImportData->AssetMetaInfo->Type == TEXT("surface") || AssetImportData->AssetMetaInfo->Type == TEXT("atlas") || AssetImportData->AssetMetaInfo->Type == TEXT("brush"))
+	{
+		AssetRecord MSRecord;
+		MSRecord.ID = AssetImportData->AssetMetaInfo->Id;
+		MSRecord.Name = SurfaceImportParams->BaseParams->AssetName;
+		MSRecord.Path = SurfaceImportParams->BaseParams->AssetDestination;
+		MSRecord.Type = AssetImportData->AssetMetaInfo->Type;
+		FAssetsDatabase::Get()->AddRecord(MSRecord);
+	}
+
+	if (MaterialInstance == nullptr) {
+		return;
+	}
+
+	const UMegascansSettings* MegascansSettings = GetDefault<UMegascansSettings>();
+	if (MegascansSettings->bApplyToSelection)
+	{
+		ApplyMaterialToSelection(MaterialInstance);
+	}
+
+}
+
+UMaterialInstanceConstant* FImportSurface::ImportSurface(TSharedPtr<FAssetTypeData> AssetImportData, TSharedPtr<SurfaceParams> SurfaceImportParams)
+{
+	// Import channel packed maps
+	auto PackedImportData = ImportPackedMaps(AssetImportData, SurfaceImportParams->TexturesDestination);
+	// Create Material Instance
+	UMaterialInstanceConstant* MaterialInstance = CreateInstanceMaterial(SurfaceImportParams);
+	
+	
+	// Get list of maps to be filetered on import. All maps that are in channel packed maps will be filetered upon import.
+	const TArray<FString> FilteredTextureTypes = GetFilteredMaps(PackedImportData, MaterialInstance);
+	// Import all the textures
+	TMap<FString, TextureData> TextureMaps = ImportTextureMaps( AssetImportData, SurfaceImportParams, FilteredTextureTypes);
+
+	// Exit if material instance creation failed.
+	if (MaterialInstance == nullptr)
+	{
+		return nullptr;
+	}
+
+	// Plug the imported textures into the Material Instance.
+	MInstanceApplyTextures(TextureMaps, MaterialInstance);
+	if (SurfaceImportParams->bContainsPackedMaps) {
+		MInstanceApplyPackedMaps(PackedImportData, MaterialInstance);
+	}
+
+	if (AssetImportData->AssetMetaInfo->bSavePackages)
+	{
+		AssetUtils::SavePackage(MaterialInstance);
+	}
+
+	return MaterialInstance;
+
+}
+
+TSharedPtr<FImportSurface> FImportSurface::Get()
+{
+	if (!ImportSurfaceInst.IsValid())
+	{
+		ImportSurfaceInst = MakeShareable(new FImportSurface);
+	}
+	return ImportSurfaceInst;
+}
+
+
+TMap<FString, TextureData> FImportSurface::ImportTextureMaps(TSharedPtr<FAssetTypeData> AssetImportData, TSharedPtr<SurfaceParams> SurfaceImportParams , const TArray<FString>& FilteredTextureTypes)
+{
+	TMap<FString, TextureData> TextureMaps;
+	for (TSharedPtr<FAssetTextureData> TextureMetaData : AssetImportData->TextureComponents)
+	{
+		
+		FString TextureType = TextureMetaData->Type;
+
+		if (FilteredTextureTypes.Contains(TextureType)) continue;
+
+		UAssetImportTask* TextureImportTask = CreateImportTask(TextureMetaData, SurfaceImportParams->TexturesDestination);
+		TextureData TextureImportData = ImportTexture(TextureImportTask);
+
+		if (TextureType == TEXT("normal"))
+		{
+			TextureImportData.TextureAsset->CompressionSettings = TextureCompressionSettings::TC_Normalmap;
+			TextureImportData.TextureAsset->bFlipGreenChannel = 1;
+		}
+
+		if (TextureType == TEXT("opacity"))
+			TextureImportData.TextureAsset->MipGenSettings = TextureMipGenSettings::TMGS_NoMipmaps;
+
+		TextureImportData.TextureAsset->VirtualTextureStreaming = 0;
+
+
+		if (NonLinearMaps.Contains(TextureType))
+		{
+			TextureImportData.TextureAsset->SRGB = 1;
+		}
+		else
+		{
+			TextureImportData.TextureAsset->SRGB = 0;
+		}
+
+		if (TextureType != TEXT("normal"))
+			TextureImportData.TextureAsset->CompressionSettings = TextureCompressionSettings::TC_BC7;
+		TextureImportData.TextureAsset->SetFlags(RF_Standalone);
+		TextureImportData.TextureAsset->MarkPackageDirty();
+		TextureImportData.TextureAsset->PostEditChange();
+		if (AssetImportData->AssetMetaInfo->bSavePackages)
+		{
+			AssetUtils::SavePackage(TextureImportData.TextureAsset);
+		}
+		
+		TextureMaps.Add(TextureType, TextureImportData);
+	}
+	return TextureMaps;
+}
+
+
+
+
+UAssetImportTask* FImportSurface::CreateImportTask(TSharedPtr<FAssetTextureData> TextureMetaData, const FString& TexturesDestination)
+{
+	FString Filename;
+	Filename = FPaths::GetBaseFilename(TextureMetaData->NameOverride);
+	UAssetImportTask* TextureImportTask = NewObject<UAssetImportTask>();
+	TextureImportTask->bAutomated = true;
+	TextureImportTask->bSave = false;
+	TextureImportTask->Filename = TextureMetaData->Path;
+	TextureImportTask->DestinationName = RemoveReservedKeywords(NormalizeString(Filename));
+	TextureImportTask->DestinationPath = TexturesDestination;
+	TextureImportTask->bReplaceExisting = true;
+	return TextureImportTask;
+}
+
+void FImportSurface::ApplyMaterialToSelection(UMaterialInstanceConstant* MaterialInstance)
+{
+	USelection* SelectedActors = GEditor->GetSelectedActors();
+	TArray<AActor*> Actors;
+	TArray<ULevel*> UniqueLevels;
+	for (FSelectionIterator Iter(*SelectedActors); Iter; ++Iter)
+	{
+		AActor* Actor = Cast<AActor>(*Iter);
+		TArray<UStaticMeshComponent*> Components;
+		Actor->GetComponents<UStaticMeshComponent>(Components);
+		for (int32 i = 0; i < Components.Num(); i++)
+		{
+			UStaticMeshComponent* MeshComponent = Components[i];
+			int32 mCnt = MeshComponent->GetNumMaterials();
+			for (int j = 0; j < mCnt; j++)
+				MeshComponent->SetMaterial(j, MaterialInstance);
+
+		}
+	}
+}
+
+TextureData FImportSurface::ImportTexture(UAssetImportTask* TextureImportTask)
+{
+	TextureData TextureImportData;
+	TArray<UAssetImportTask*> ImportTasks;
+	IAssetTools& AssetTools = FModuleManager::GetModuleChecked<FAssetToolsModule>("AssetTools").Get();
+	ImportTasks.Add(TextureImportTask);
+
+	AssetTools.ImportAssetTasks(ImportTasks);
+	for (UAssetImportTask* ImpTask : ImportTasks)
+	{
+		if (ImpTask->ImportedObjectPaths.Num() > 0)
+		{
+			TextureImportData.TextureAsset = CastChecked<UTexture>(UEditorAssetLibrary::LoadAsset(ImpTask->ImportedObjectPaths[0]));
+			TextureImportData.Path = ImpTask->ImportedObjectPaths[0];
+		}
+	}
+	return TextureImportData;
+}
+
+//UMaterialInstanceConstant* FImportSurface::CreateInstanceMaterial(const FString & MasterMaterialPath, const FString& InstanceDestination, const FString& MInstanceName)
+//{
+//	if (!UEditorAssetLibrary::DoesAssetExist(MasterMaterialPath)) return nullptr;
+//
+//	UMaterialInterface* MasterMaterial = CastChecked<UMaterialInterface>(LoadAsset(MasterMaterialPath));
+//	IAssetTools& AssetTools = FModuleManager::GetModuleChecked<FAssetToolsModule>("AssetTools").Get();
+//	UMaterialInstanceConstantFactoryNew* Factory = NewObject<UMaterialInstanceConstantFactoryNew>();
+//
+//	UObject* MaterialInstAsset = AssetTools.CreateAsset(MInstanceName, InstanceDestination, UMaterialInstanceConstant::StaticClass(), Factory);
+//	UMaterialInstanceConstant *MaterialInstance = CastChecked<UMaterialInstanceConstant>(MaterialInstAsset);
+//	
+//	UMaterialEditingLibrary::SetMaterialInstanceParent(MaterialInstance, MasterMaterial);
+//	if (MaterialInstance)
+//	{
+//		MaterialInstance->SetFlags(RF_Standalone);
+//		MaterialInstance->MarkPackageDirty();
+//		MaterialInstance->PostEditChange();
+//
+//		
+//	}
+//	
+//	return MaterialInstance;
+//}
+
+UMaterialInstanceConstant* FImportSurface::CreateInstanceMaterial(TSharedPtr<SurfaceParams> SurfaceImportParams)
+{
+	if (!UEditorAssetLibrary::DoesAssetExist(SurfaceImportParams->MasterMaterialPath)) return nullptr;
+	UMaterialInterface* MasterMaterial = CastChecked<UMaterialInterface>(LoadAsset(SurfaceImportParams->MasterMaterialPath));
+	IAssetTools& AssetTools = FModuleManager::GetModuleChecked<FAssetToolsModule>("AssetTools").Get();
+	UMaterialInstanceConstantFactoryNew* Factory = NewObject<UMaterialInstanceConstantFactoryNew>();
+	UObject* MaterialInstAsset = AssetTools.CreateAsset(SurfaceImportParams->MaterialInstanceName, SurfaceImportParams->MaterialInstanceDestination, UMaterialInstanceConstant::StaticClass(), Factory);
+	UMaterialInstanceConstant* MaterialInstance = CastChecked<UMaterialInstanceConstant>(MaterialInstAsset);
+	UMaterialEditingLibrary::SetMaterialInstanceParent(MaterialInstance, MasterMaterial);
+	if (MaterialInstance)
+	{
+		MaterialInstance->SetFlags(RF_Standalone);
+		MaterialInstance->MarkPackageDirty();
+		MaterialInstance->PostEditChange();
+	}
+	return MaterialInstance;
+}
+
+void FImportSurface::MInstanceApplyTextures(TMap<FString, TextureData> TextureMaps, UMaterialInstanceConstant* MaterialInstance)
+{	
+
+	for (const TPair<FString, TextureData>& TextureData : TextureMaps)
+	{
+
+		if (UMaterialEditingLibrary::GetMaterialInstanceTextureParameterValue(MaterialInstance, FName(*TextureData.Key)))
+		{
+			//UTexture* TextureAsset = Cast<UTexture>(LoadAsset(TextureData.Value.Path));
+			UTexture* TextureAsset = TextureData.Value.TextureAsset;
+			if (UMaterialEditingLibrary::SetMaterialInstanceTextureParameterValue(MaterialInstance, FName(*TextureData.Key), TextureAsset))
+			{
+				MaterialInstance->SetFlags(RF_Standalone);
+				MaterialInstance->MarkPackageDirty();
+				MaterialInstance->PostEditChange();
+				
+			}
+			
+		}
+		
+	}
+	
+}
+
+TArray<FString> FImportSurface::GetPackedMapsList(TSharedPtr<FAssetTypeData> AssetImportData)
+{
+	TArray<FString> PackedMaps;
+	return PackedMaps;
+}
+
+
+//TMap<FString, TSharedPtr<FAssetPackedTextures>> FImportSurface::ImportPackedMaps(TSharedPtr<FAssetTypeData> AssetImportData, const FString& TexturesDestination)
+//{
+//	TMap<FString, TSharedPtr<FAssetPackedTextures>> PackedImportData;
+//	for (TSharedPtr<FAssetPackedTextures> PackedData : AssetImportData->PackedTextures)
+//	{
+//		UAssetImportTask* TextureImportTask = CreateImportTask(PackedData->PackedTextureData, TexturesDestination);
+//		TextureData TextureImportData = ImportTexture(TextureImportTask);
+//
+//		//UTexture* TextureAsset = Cast<UTexture>(LoadAsset(TextureImportData.Path));
+//		TextureImportData.TextureAsset->SRGB = 0;
+//
+//		FString ImportedMap = TextureImportData.Path;
+//		if (PackedImportData.Contains(ImportedMap))
+//		{
+//			PackedImportData.Remove(ImportedMap);
+//		}
+//		PackedImportData.Add(ImportedMap, PackedData);
+//		if (AssetImportData->AssetMetaInfo->bSavePackages)
+//		{
+//			AssetUtils::SavePackage(TextureImportData.TextureAsset);
+//		}
+//	}
+//	return PackedImportData;
+//}
+
+TMap<FString, TSharedPtr<FAssetPackedTextures>> FImportSurface::ImportPackedMaps(TSharedPtr<FAssetTypeData> AssetImportData, const FString& TexturesDestination)
+{
+	TMap<FString, TSharedPtr<FAssetPackedTextures>> PackedImportData;
+	for (TSharedPtr<FAssetPackedTextures> PackedData : AssetImportData->PackedTextures)
+	{
+		UAssetImportTask* TextureImportTask = CreateImportTask(PackedData->PackedTextureData, TexturesDestination);
+		TextureData TextureImportData = ImportTexture(TextureImportTask);
+
+		//UTexture* TextureAsset = Cast<UTexture>(LoadAsset(TextureImportData.Path));
+		TextureImportData.TextureAsset->SRGB = 0;
+
+		FString ImportedMap = TextureImportData.Path;
+		if (PackedImportData.Contains(ImportedMap))
+		{
+			PackedImportData.Remove(ImportedMap);
+		}
+		PackedImportData.Add(ImportedMap, PackedData);
+
+		if (AssetImportData->AssetMetaInfo->bSavePackages)
+		{
+			AssetUtils::SavePackage(TextureImportData.TextureAsset);
+		}
+		
+	}
+	return PackedImportData;
+}
+
+TArray<FString> FImportSurface::GetPackedTypes(ChannelPackedData PackedImportData)
+{
+	TArray<FString> PackedTypes;
+	for (auto& PackedData : PackedImportData)
+	{
+		for (auto& ChData : PackedData.Value->ChannelData)
+		{			
+				if (ChData.Value[0] == TEXT("gray") || ChData.Value[0] == TEXT("empty") || ChData.Value[0] == TEXT("value")) continue;
+				PackedTypes.Add(ChData.Value[0].ToLower());
+				
+		}
+
+	}
+	return PackedTypes;
+}
+
+void FImportSurface::MInstanceApplyPackedMaps(TMap<FString, TSharedPtr<FAssetPackedTextures>> PackedImportData, UMaterialInstanceConstant* MaterialInstance)
+{
+	
+	for (auto& PackedMapData : PackedImportData)
+	{
+		FString ChannelPackedType = TEXT("");
+		FString PackedMap = PackedMapData.Key;
+		UTexture* PackedAsset = Cast<UTexture>(UEditorAssetLibrary::LoadAsset(PackedMapData.Key));
+		if (PackedAsset == nullptr) continue;
+		for (auto& ChData : PackedMapData.Value->ChannelData)
+		{
+			if (ChData.Value[0] == TEXT("gray") || ChData.Value[0] == TEXT("empty") || ChData.Value[0] == TEXT("value")) continue;
+			ChannelPackedType = ChannelPackedType + ChData.Value[0].Left(1);
+		}
+
+		
+		if (UMaterialEditingLibrary::GetMaterialInstanceTextureParameterValue(MaterialInstance, FName(*ChannelPackedType)))
+		{
+			UMaterialEditingLibrary::SetMaterialInstanceTextureParameterValue(MaterialInstance, FName(*ChannelPackedType), PackedAsset);
+			MaterialInstance->SetFlags(RF_Standalone);
+			MaterialInstance->MarkPackageDirty();
+			MaterialInstance->PostEditChange();
+		}
+	}
+
+}
+
+
+
+FString FImportSurface::GetMaterialOverride(TSharedPtr<FAssetTypeData> AssetImportData)
+{
+	TArray<FString> SupportedMatOverrides = { "3d", "3dplant", "surface" };
+	FString AssetType = AssetImportData->AssetMetaInfo->Type;
+	const UMaterialPresetsSettings* MatOverrideSettings = GetDefault< UMaterialPresetsSettings>();
+	if (SupportedMatOverrides.Contains(AssetType))
+	{
+		if (AssetType == "surface" && MatOverrideSettings->MasterMaterialSurface != nullptr) return MatOverrideSettings->MasterMaterialSurface->GetPathName();
+		if (AssetType == "3d" && MatOverrideSettings->MasterMaterial3d != nullptr) return MatOverrideSettings->MasterMaterial3d->GetPathName();
+		if (AssetType == "3dplant" && MatOverrideSettings->MasterMaterialPlant != nullptr) return MatOverrideSettings->MasterMaterialPlant->GetPathName();
+
+	}
+	return TEXT("");
+
+}
+
+
+FString FImportSurface::GetSurfaceType(TSharedPtr<FAssetTypeData> AssetImportData)
+{
+	FString SurfaceType = TEXT("");
+	FString AssetType = AssetImportData->AssetMetaInfo->Type;
+	auto& TagsList = AssetImportData->AssetMetaInfo->Tags;
+	auto& CategoriesList = AssetImportData->AssetMetaInfo->Categories;
+
+	if(TagsList.Contains(TEXT("metal"))) SurfaceType = TEXT("metal");
+	else if(TagsList.Contains(TEXT("surface imperfection"))) SurfaceType = TEXT("imperfection");
+	else if (TagsList.Contains(TEXT("tileable displacement"))) SurfaceType = TEXT("displacement");
+	
+	else if (CategoriesList[0] == TEXT("brush")) SurfaceType = TEXT("brush");
+	else if (CategoriesList[0] == TEXT("atlas") && TagsList.Contains(TEXT("decal"))) SurfaceType = TEXT("decal");
+	else if (CategoriesList[0] == TEXT("atlas") && !TagsList.Contains(TEXT("decal"))) SurfaceType = TEXT("atlas");
+
+	return SurfaceType;
+}
+
+
+
+TArray<FString> FImportSurface::GetFilteredMaps(ChannelPackedData PackedImportData, UMaterialInstanceConstant* MaterialInstance)
+{
+	TArray<FString> FilteredMaps = GetPackedTypes(PackedImportData);
+	const UMegascansSettings* MegascansSettings = GetDefault<UMegascansSettings>();
+	bool bFalsifyThis = false;	
+	if (MaterialInstance !=nullptr && MegascansSettings->bFilterMasterMaterialMaps)
+	{		
+		for (FString MapType : AllMapTypes)
+		{
+			if (UMaterialEditingLibrary::GetMaterialInstanceTextureParameterValue(MaterialInstance, FName(*MapType)) == nullptr)
+			{
+				FilteredMaps.Add(MapType);
+			}
+		}
+
+	}
+
+	return FilteredMaps;
+
+
+}
+
+
+
diff --git a/Engine/Plugins/MegascansPlugin/Source/MegascansPlugin/Private/AssetImporters/ImportSurface.h b/Engine/Plugins/MegascansPlugin/Source/MegascansPlugin/Private/AssetImporters/ImportSurface.h
new file mode 100644
--- /dev/null
+++ b/Engine/Plugins/MegascansPlugin/Source/MegascansPlugin/Private/AssetImporters/ImportSurface.h
@@ -0,0 +1,90 @@
+// Copyright Epic Games, Inc. All Rights Reserved.
+#pragma once
+#include "CoreMinimal.h"
+#include "AssetImportData.h"
+
+#include "ImportFactory.h"
+#include "Utilities/AssetData.h"
+
+
+class UMaterialInstanceConstant;
+class UTexture;
+
+struct TextureData
+{
+	FString Path;
+	UTexture* TextureAsset;
+};
+
+struct SurfaceImportParams
+{
+	
+	FString TexturesDestination;
+	FString AssetName;
+	//FString MasterMaterialPath;
+	FString MInstanceDestination;
+	FString AssetDestination;
+	FString MeshDestination;
+	FString FoliageDestination;
+	FString MInstanceName;
+	bool bExrEnabled;
+};
+
+
+class FImportSurface : public IImportAsset
+{
+	typedef TMap<FString, TSharedPtr<FAssetPackedTextures>> ChannelPackedData;
+
+
+private:
+	FImportSurface() = default;
+	bool bEnableExrDisplacement;	
+	static TSharedPtr<FImportSurface> ImportSurfaceInst;
+	TextureData ImportTexture(UAssetImportTask * TextureImportTask);
+	UMaterialInstanceConstant* CreateInstanceMaterial(const FString& MasterMaterialPath, const FString& InstanceDestination, const FString& AssetName);
+	void MInstanceApplyTextures(TMap<FString, TextureData> TextureMaps, UMaterialInstanceConstant* MaterialInstance);
+	TArray<FString> GetPackedMapsList(TSharedPtr<FAssetTypeData> AssetImportData);
+	TMap<FString, TSharedPtr<FAssetPackedTextures>> ImportPackedMaps(TSharedPtr<FAssetTypeData> AssetImportData, const FString& TexturesDestination);
+	FString GetSurfaceType(TSharedPtr<FAssetTypeData> AssetImportData);		
+	TArray<FString> GetFilteredMaps(ChannelPackedData PackedImportData, UMaterialInstanceConstant* MaterialInstance);
+	const TArray<FString> NonLinearMaps = { "albedo", "diffuse", "translucency", "specular" };
+	const TArray<FString> ColorPackedMaps = { "albedo", "diffuse", "translucency", "specular" , "normal" };
+
+
+	const TArray<FString> AllMapTypes = {
+		"opacity",
+		  "normal",
+		  "bump",
+		  "gloss" ,
+		  "displacement",
+		  "ao",
+		  "thickness",
+		  "cavity",
+		  "specular",
+		  "normalbump",
+		  "translucency",
+		  "albedo",
+		  "roughness",
+		  "diffuse",
+		  "fuzz",
+		  "metalness"
+	};
+
+	TArray<FString> GetPackedTypes(ChannelPackedData PackedImportData);
+	void MInstanceApplyPackedMaps(TMap<FString, TSharedPtr<FAssetPackedTextures>> PackedImportData, UMaterialInstanceConstant* MaterialInstance);
+	
+
+	FString  GetMaterialOverride(TSharedPtr<FAssetTypeData> AssetImportData);
+	// New implementation
+	//TMap<FString, TSharedPtr<FAssetPackedTextures>> ImportPackedMaps(TArray<TSharedPtr<FAssetPackedTextures>> PackedTextures, const FString& TexturesDestination);
+	UMaterialInstanceConstant* CreateInstanceMaterial(TSharedPtr<SurfaceParams> SurfaceImportParams);
+	TMap<FString, TextureData> ImportTextureMaps(TSharedPtr<FAssetTypeData> AssetImportData, TSharedPtr<SurfaceParams> SurfaceImportParams, const TArray<FString>& FilteredTextureTypes);
+	UAssetImportTask* CreateImportTask(TSharedPtr<FAssetTextureData> TextureMetaData, const FString& TexturesDestination);
+	void ApplyMaterialToSelection(UMaterialInstanceConstant* MaterialInstance);
+
+public:
+	virtual void ImportAsset(TSharedPtr<FAssetTypeData> AssetImportData) override;
+	static TSharedPtr<FImportSurface> Get();	
+	
+	UMaterialInstanceConstant* ImportSurface( TSharedPtr<FAssetTypeData> AssetImportData, TSharedPtr<SurfaceParams> SurfaceImportParams);
+};
\ No newline at end of file
diff --git a/Engine/Plugins/MegascansPlugin/Source/MegascansPlugin/Private/AssetsImportController.cpp b/Engine/Plugins/MegascansPlugin/Source/MegascansPlugin/Private/AssetsImportController.cpp
new file mode 100644
--- /dev/null
+++ b/Engine/Plugins/MegascansPlugin/Source/MegascansPlugin/Private/AssetsImportController.cpp
@@ -0,0 +1,128 @@
+// Copyright Epic Games, Inc. All Rights Reserved.
+#include "AssetsImportController.h"
+#include "AssetImportDataHandler.h"
+#include "AssetImporters/ImportSurface.h"
+#include "AssetImporters/Import3D.h"
+#include "AssetImporters/Import3DPlant.h"
+#include "UI/MSSettings.h"
+#include "Utilities/AssetsDatabase.h"
+#include "Utilities/Analytics.h"
+#include "Utilities/AssetData.h"
+#include "AssetImporters/ImportFactory.h"
+
+#include "Misc/MessageDialog.h"
+#include "Runtime/Core/Public/Internationalization/Text.h"
+#include "Misc/Paths.h"
+
+
+TSharedPtr<FAssetsImportController> FAssetsImportController::AssetsImportController;
+
+
+// Get instance of the class
+TSharedPtr<FAssetsImportController> FAssetsImportController::Get()
+{
+	if (!AssetsImportController.IsValid())
+	{
+		AssetsImportController = MakeShareable(new FAssetsImportController);	
+	}
+	return AssetsImportController;
+}
+
+
+
+void FAssetsImportController::DataReceived(const FString DataFromBridge)
+{
+	
+	//UE_LOG(LogTemp, Error, TEXT("Data from Bridge :%s"), *DataFromBridge);
+	if (IsGarbageCollecting() || GIsSavingPackage) return;
+	TArray<FDHIData> DHIAssetsData;
+	if (DHI::GetDHIJsonData(DataFromBridge, DHIAssetsData)) {
+		for (FDHIData CharacterData : DHIAssetsData) {
+			DHI::CopyCharacter(CharacterData);	
+		}
+	}
+	else ImportAssets(DataFromBridge);	
+}
+
+// Gets import preferences from a json file, parses the Bridge json and calls the appropriate import function based in asset type
+void FAssetsImportController::ImportAssets(const FString& AssetsImportJson)
+{
+	bool bSavePackages = false;
+	bool bSkipImportAll = false;
+	bool bImportAll = false;
+	bool bAllSkipOrImport = false;
+
+	TSharedPtr<FAssetsData> AssetsImportData = FAssetDataHandler::Get()->GetAssetsData(AssetsImportJson);
+	
+
+	checkf(AssetsImportData->AllAssetsData.Num() > 0, TEXT("There was an error reading asset data."));	
+	const UMegascansSettings* MegascansSettings = GetDefault<UMegascansSettings>();
+	if (AssetsImportData->AllAssetsData.Num() > 10)
+	{
+		bSavePackages = true;
+		if (MegascansSettings->bBatchImportPrompt)
+		{			
+			EAppReturnType::Type ContinueImport = FMessageDialog::Open(EAppMsgType::OkCancel, FText(FText::FromString("You are about to download more than 10 assets. Press Ok to continue.")));
+			if (ContinueImport == EAppReturnType::Cancel) return;
+			
+		}
+	}
+
+	TSharedPtr<FJsonObject> UIAnalytics = FAnalytics::Get()->GenerateAnalyticsJson();
+	FAnalytics::Get()->SendAnalytics(UIAnalytics);
+	
+
+	for (TSharedPtr<FAssetTypeData> AssetImportData : AssetsImportData->AllAssetsData)
+	{	
+		//AssetImportData->AssetMetaInfo->bSavePackages = bSavePackages;
+		AssetImportData->AssetMetaInfo->bSavePackages = false;
+		TSharedPtr<FAssetImportParams> AssetSetupParameters = FAssetImportParams::Get();
+		AssetRecord Record;
+		if (FAssetsDatabase::Get()->RecordExists(AssetImportData->AssetMetaInfo->Id, Record) && FPaths::DirectoryExists(FPaths::Combine(FPaths::ProjectContentDir(), Record.Path.Replace(TEXT("/Game"), TEXT("")))))
+		{
+			if (bSkipImportAll)
+			{
+				continue;
+			}
+			if (!bAllSkipOrImport)
+			{
+				EAppReturnType::Type ReimportAssetDlg = FMessageDialog::Open(EAppMsgType::YesNoYesAllNoAll, FText(FText::FromString(FString::Printf(TEXT("The asset %s already exists at %s. Do you want to import this asset.?"), *AssetImportData->AssetMetaInfo->Name, *Record.Path))));
+				if (ReimportAssetDlg == EAppReturnType::No) {
+					continue;
+				}
+				if (ReimportAssetDlg == EAppReturnType::NoAll)
+				{
+					bSkipImportAll = true;
+					bAllSkipOrImport = true;
+					continue;
+				}
+				if (ReimportAssetDlg == EAppReturnType::YesAll)
+				{
+					bAllSkipOrImport = true;
+					bImportAll = true;
+				}
+			}
+
+
+		}
+		if (AssetImportData->AssetMetaInfo->Type == "3d")
+		{
+			
+			FImport3d::Get()->ImportAsset(AssetImportData);			
+		}
+		else if (AssetImportData->AssetMetaInfo->Type == "3dplant")
+		{
+			FImportPlant::Get()->ImportAsset(AssetImportData);
+		}
+
+		else if (AssetImportData->AssetMetaInfo->Type == "surface" || AssetImportData->AssetMetaInfo->Type == "atlas" || AssetImportData->AssetMetaInfo->Type == "brush")
+		{			
+			FImportSurface::Get()->ImportAsset(AssetImportData);
+		}
+
+	}	
+	AssetsImportData.Reset();	
+	FImport3d::Get().Reset();	
+	FImportPlant::Get().Reset();
+	FImportSurface::Get().Reset();
+}
\ No newline at end of file
diff --git a/Engine/Plugins/MegascansPlugin/Source/MegascansPlugin/Private/AssetsImportController.h b/Engine/Plugins/MegascansPlugin/Source/MegascansPlugin/Private/AssetsImportController.h
new file mode 100644
--- /dev/null
+++ b/Engine/Plugins/MegascansPlugin/Source/MegascansPlugin/Private/AssetsImportController.h
@@ -0,0 +1,19 @@
+// Copyright Epic Games, Inc. All Rights Reserved.
+#pragma once
+#include "CoreMinimal.h"
+
+
+
+class FAssetsImportController
+
+{
+private:
+	FAssetsImportController() = default;
+	static TSharedPtr<FAssetsImportController> AssetsImportController;
+	void ImportAssets(const FString & AssetsImportJson);
+
+public:	
+	static TSharedPtr<FAssetsImportController> Get();
+	void DataReceived(const FString DataFromBridge);
+
+};
\ No newline at end of file
diff --git a/Engine/Plugins/MegascansPlugin/Source/MegascansPlugin/Private/MSPythonBridge.cpp b/Engine/Plugins/MegascansPlugin/Source/MegascansPlugin/Private/MSPythonBridge.cpp
new file mode 100644
--- /dev/null
+++ b/Engine/Plugins/MegascansPlugin/Source/MegascansPlugin/Private/MSPythonBridge.cpp
@@ -0,0 +1,19 @@
+// Copyright Epic Games, Inc. All Rights Reserved.
+#include "MSPythonBridge.h"
+
+UMSPythonBridge* UMSPythonBridge::Get()
+{
+	TArray<UClass*> PythonBridgeClasses;
+	GetDerivedClasses(UMSPythonBridge::StaticClass(), PythonBridgeClasses);
+	int32 NumClasses = PythonBridgeClasses.Num();
+	if (NumClasses > 0)
+	{
+		return Cast<UMSPythonBridge>(PythonBridgeClasses[NumClasses - 1]->GetDefaultObject());
+	}
+	return nullptr;
+};
+
+void UMSPythonBridge::CalledFromPython(FString InputString)
+{
+	UE_LOG(LogTemp, Error, TEXT("Hello from Python"));
+}
\ No newline at end of file
diff --git a/Engine/Plugins/MegascansPlugin/Source/MegascansPlugin/Private/MSPythonBridge.h b/Engine/Plugins/MegascansPlugin/Source/MegascansPlugin/Private/MSPythonBridge.h
new file mode 100644
--- /dev/null
+++ b/Engine/Plugins/MegascansPlugin/Source/MegascansPlugin/Private/MSPythonBridge.h
@@ -0,0 +1,37 @@
+// Copyright Epic Games, Inc. All Rights Reserved.
+#pragma once
+#include "MSPythonBridge.generated.h"
+
+UCLASS(Blueprintable)
+class UMSPythonBridge : public UObject
+{
+	GENERATED_BODY()
+public:
+	UFUNCTION(BlueprintCallable, Category = Python)
+		static UMSPythonBridge * Get();
+
+	UFUNCTION(BlueprintImplementableEvent, Category = Python)
+		void InitializePythonWindow() const;
+
+	UFUNCTION(BlueprintImplementableEvent, Category = Python)
+	void GetUeData(const FString& CommandName) const;
+
+	UFUNCTION(BlueprintImplementableEvent, Category = Python)
+		void TestFbxImport() const;
+
+	UPROPERTY(EditAnywhere, BlueprintReadWrite, Category = Python)
+		FString JsonData;
+
+	//Testing bridge from Python to C++
+	UFUNCTION(BlueprintCallable, Category = Python)
+		static void CalledFromPython(FString InputString);
+
+	//UPROPERTY(EditAnywhere, BlueprintReadWrite)
+	//FString jsondata;
+	//TArray<FString> filenames;
+
+	//UFUNCTION(BlueprintImplementableEvent, Category = Python)
+	//void DownloadByTestFilenames() const;
+
+};
+
diff --git a/Engine/Plugins/MegascansPlugin/Source/MegascansPlugin/Private/MegascansLiveLinkModule.cpp b/Engine/Plugins/MegascansPlugin/Source/MegascansPlugin/Private/MegascansLiveLinkModule.cpp
new file mode 100644
--- /dev/null
+++ b/Engine/Plugins/MegascansPlugin/Source/MegascansPlugin/Private/MegascansLiveLinkModule.cpp
@@ -0,0 +1,50 @@
+// Copyright Epic Games, Inc. All Rights Reserved.
+#include "IMegascansLiveLinkModule.h"
+#include "UI/QMSUIManager.h"
+#include "UI/MaterialBlendingDetails.h"
+#include "UI/MSSettings.h"
+
+
+
+
+#define LOCTEXT_NAMESPACE "MegascansPlugin"
+
+class FQMSLiveLinkModule : public IMegascansLiveLinkModule
+{
+public:
+	virtual void StartupModule() override
+	{
+		if (GIsEditor && !IsRunningCommandlet())
+		{
+			 FQMSUIManager::Initialize();
+		}
+
+		auto& PropertyModule = FModuleManager::LoadModuleChecked<FPropertyEditorModule>("PropertyEditor");
+		PropertyModule.RegisterCustomClassLayout(
+			"MaterialBlendSettings",
+			FOnGetDetailCustomizationInstance::CreateStatic(&BlendSettingsCustomization::MakeInstance)
+		);
+		PropertyModule.NotifyCustomizationModuleChanged();
+	}
+
+	virtual void ShutdownModule() override
+	{
+		if ((GIsEditor && !IsRunningCommandlet()))
+		{
+			
+		}
+
+		if (FModuleManager::Get().IsModuleLoaded("PropertyEditor"))
+		{
+			auto& PropertyModule = FModuleManager::LoadModuleChecked<FPropertyEditorModule>("PropertyEditor");
+
+			PropertyModule.UnregisterCustomClassLayout("MaterialBlendSettings");
+		}
+	}
+
+
+};
+
+IMPLEMENT_MODULE(FQMSLiveLinkModule, MegascansPlugin);
+
+#undef LOCTEXT_NAMESPACE
diff --git a/Engine/Plugins/MegascansPlugin/Source/MegascansPlugin/Private/TCPServerImp.cpp b/Engine/Plugins/MegascansPlugin/Source/MegascansPlugin/Private/TCPServerImp.cpp
new file mode 100644
--- /dev/null
+++ b/Engine/Plugins/MegascansPlugin/Source/MegascansPlugin/Private/TCPServerImp.cpp
@@ -0,0 +1,183 @@
+// Copyright Epic Games, Inc. All Rights Reserved.
+#include "TCPServerImp.h"
+#include <string>
+#include "NetworkMessage.h"
+
+#include "Runtime/Core/Public/Misc/ScopedSlowTask.h"
+
+
+FTCPServer::FTCPServer() 
+{
+	FThreadSafeCounter  WorkerCounter;
+	FString ThreadName(FString::Printf(TEXT("MegascansPlugin%i"), WorkerCounter.Increment()));
+	ClientThread = FRunnableThread::Create(this, *ThreadName, 8 * 1024, TPri_Normal);
+	
+}
+
+
+FTCPServer::~FTCPServer()
+{
+	
+	Stop();
+
+	
+	if (Listener != NULL)
+	{
+		Listener->Stop();
+		delete Listener;
+		Listener = NULL;
+	}
+	
+	if (!PendingClients.IsEmpty())
+	{
+		FSocket *Client = NULL;
+
+		while (PendingClients.Dequeue(Client))
+		{
+			Client->Close();
+			delete Client;
+			Client = NULL;
+		}
+	}
+
+	for (TArray<class FSocket*>::TIterator ClientIt(Clients); ClientIt; ++ClientIt)
+	{
+		(*ClientIt)->Close();
+		delete (*ClientIt);
+		(*ClientIt) = NULL;
+	}
+	
+	if (ClientThread != NULL)
+	{
+		ClientThread->Kill(true);
+		delete ClientThread;
+	}
+}
+
+
+bool FTCPServer::Init()
+{	
+
+	if (Listener == NULL)
+	{
+		FIPv4Address address;
+		FIPv4Address::Parse(LocalHostIP, address);
+		FIPv4Endpoint endPoint = FIPv4Endpoint(address, PortNum);
+		Listener = new FTcpListener(endPoint, FTimespan::FromMilliseconds(300));
+		Listener->OnConnectionAccepted().BindRaw(this, &FTCPServer::HandleListenerConnectionAccepted);
+		Stopping = false;
+	}
+	
+	return (Listener != NULL);
+}
+
+
+uint32  FTCPServer::Run()
+{
+	while (!Stopping)
+	{
+
+		if (!PendingClients.IsEmpty())
+		{
+			
+			for (TArray<class FSocket*>::TIterator ClientIt(Clients); ClientIt; ++ClientIt)
+			{
+				(*ClientIt)->Close();
+				delete (*ClientIt);
+				(*ClientIt) = NULL;
+			}
+
+			FSocket *Client = NULL;
+			while (PendingClients.Dequeue(Client))
+			{
+				
+				Clients.Empty();
+				Clients.Add(Client);
+				ConnectionTimer.AddZeroed(1);
+			}
+		}
+
+		
+
+		if (Clients.Num() > 0)
+		{
+			
+			FString RecievedJson = TEXT("");
+			uint32 DataSize = 0;
+
+			for (TArray<class FSocket*>::TIterator ClientIt(Clients); ClientIt; ++ClientIt)
+			{
+				
+				FSocket *Client = *ClientIt;				
+				DataSize = 0;
+				FString Request;
+				TArray<uint8> data;
+				bool haveMessage = false;
+				while (Client->HasPendingData(DataSize))
+				{					
+					haveMessage = RecvMessage(Client, DataSize, Request);										
+					if (haveMessage)
+					{					
+
+						int32 timer = 0;						
+						RecievedJson += Request;						
+						Request.Empty();
+						haveMessage = false;
+					}
+
+				}		
+
+			}
+
+			if (RecievedJson != TEXT(""))
+			{
+				if (!IsGarbageCollecting() && !GIsSavingPackage)
+				{			
+
+					AsyncTask(ENamedThreads::GameThread, [this, RecievedJson]() {
+						FAssetsImportController::Get()->DataReceived(RecievedJson);
+					});
+					RecievedJson.Empty();
+				}
+			}
+			
+		}
+		
+		
+		FPlatformProcess::Sleep(0.3f); //0.3f
+	}
+
+	return 0;
+}
+
+
+bool FTCPServer::RecvMessage(FSocket *Socket, uint32 DataSize, FString& Message)
+{
+	check(Socket);	
+	FArrayReaderPtr Datagram = MakeShareable(new FArrayReader(true));
+	int32 stuff = 16;
+	Datagram->Init(FMath::Min(DataSize, 65507u), 81920);
+	int32 BytesRead = 0;
+
+	if (Socket->Recv(Datagram->GetData(), Datagram->Num(), BytesRead))
+	{
+		
+		char* Data = (char*)Datagram->GetData();
+		Data[BytesRead] = '\0';
+		FString message = UTF8_TO_TCHAR(Data);
+		Message = message;
+
+		return true;
+	}
+
+	return false;
+}
+
+bool FTCPServer::HandleListenerConnectionAccepted(class FSocket *ClientSocket, const FIPv4Endpoint& ClientEndpoint)
+
+{
+	PendingClients.Enqueue(ClientSocket);
+	return true;
+}
+
+
diff --git a/Engine/Plugins/MegascansPlugin/Source/MegascansPlugin/Private/TCPServerImp.h b/Engine/Plugins/MegascansPlugin/Source/MegascansPlugin/Private/TCPServerImp.h
new file mode 100644
--- /dev/null
+++ b/Engine/Plugins/MegascansPlugin/Source/MegascansPlugin/Private/TCPServerImp.h
@@ -0,0 +1,58 @@
+// Copyright Epic Games, Inc. All Rights Reserved.
+#pragma once
+
+
+#include "CoreMinimal.h"
+#include <string>
+#include "Runtime/Sockets/Public/Sockets.h"
+#include "Runtime/Core/Public/Misc/DateTime.h"
+#include "Runtime/Core/Public/Async/Async.h"
+#include "MSPythonBridge.h"
+#include "UI/QMSUIManager.h"
+#include "AssetsImportController.h"
+
+#include "Interfaces/IPv4/IPv4Address.h"
+#include "Interfaces/IPv4/IPv4Endpoint.h"
+#include "Common/TcpListener.h"
+#include "Sockets.h"
+#include "Common/UdpSocketReceiver.h"
+
+
+
+class FTCPServer : public FRunnable
+{
+
+public:	
+	FTCPServer();
+	~FTCPServer();	
+	virtual bool Init() override;	
+	virtual uint32 Run() override;
+
+
+
+	virtual void Stop() override
+	{
+		Stopping = true;
+	}
+
+	virtual void Exit() override { }
+	bool RecvMessage(FSocket *Socket, uint32 DataSize, FString& Message);
+	bool HandleListenerConnectionAccepted(class FSocket *ClientSocket, const FIPv4Endpoint& ClientEndpoint);
+
+
+	FSocket* ListenerSocket;
+	FString LocalHostIP = "127.0.0.1";
+	int32 PortNum = 13429; 
+	int32 ConnectionTimeout;
+	TArray<class FSocket*> Clients;
+
+private:	
+	TQueue<class FSocket*, EQueueMode::Mpsc> PendingClients;
+	bool Stopping;
+	FRunnableThread* ClientThread = NULL;
+	class FTcpListener *Listener = NULL;	
+	TArray<int32> ConnectionTimer;
+
+};
+
+
diff --git a/Engine/Plugins/MegascansPlugin/Source/MegascansPlugin/Private/Tools/BlendMaterials.cpp b/Engine/Plugins/MegascansPlugin/Source/MegascansPlugin/Private/Tools/BlendMaterials.cpp
new file mode 100644
--- /dev/null
+++ b/Engine/Plugins/MegascansPlugin/Source/MegascansPlugin/Private/Tools/BlendMaterials.cpp
@@ -0,0 +1,125 @@
+// Copyright Epic Games, Inc. All Rights Reserved.
+#include "Tools/BlendMaterials.h"
+#include "EditorAssetLibrary.h"
+#include "Utilities/MiscUtils.h"
+#include "Utilities/MaterialUtils.h"
+#include "UI/MSSettings.h"
+#include "Utilities/Analytics.h"
+
+#include "Editor/UnrealEd/Classes/Factories/MaterialInstanceConstantFactoryNew.h"
+#include "Runtime/Engine/Classes/Materials/MaterialInstanceConstant.h"
+#include "Runtime/Engine/Classes/Engine/Texture.h"
+#include "Runtime/CoreUObject/Public/Misc/PackageName.h"
+#include "MaterialEditingLibrary.h"
+
+#include "Runtime/Core/Public/Misc/MessageDialog.h"
+#include "Runtime/Core/Public/Internationalization/Text.h"
+
+
+TSharedPtr<FMaterialBlend> FMaterialBlend::MaterialBlendInst;
+
+
+
+
+TSharedPtr<FMaterialBlend> FMaterialBlend::Get()
+{
+	if (!MaterialBlendInst.IsValid())
+	{
+		MaterialBlendInst = MakeShareable(new FMaterialBlend);
+	}
+	return MaterialBlendInst;
+}
+
+void FMaterialBlend::BlendSelectedMaterials()
+{
+	TSharedPtr<FJsonObject> BlendAnalytics = FAnalytics::Get()->GenerateBlendAnalyticsJson();
+	FAnalytics::Get()->SendAnalytics(BlendAnalytics);
+
+	FString FailueReason;
+	const UMaterialBlendSettings* BlendSettings = GetDefault<UMaterialBlendSettings>();
+
+	TArray<UMaterialInstanceConstant*> SelectedMaterialInstances = AssetUtils::GetSelectedAssets(TEXT("MaterialInstanceConstant"));
+
+	if (SelectedMaterialInstances.Num() < 2)
+	{
+		FMessageDialog::Open(EAppMsgType::Ok, FText(FText::FromString("Select two or more material instances to perform this operation.")));
+		return;
+	}
+	if (SelectedMaterialInstances.Num() > 3)
+	{
+		FMessageDialog::Open(EAppMsgType::Ok, FText(FText::FromString("Current Material Blending setup doesn't support more than 3 material instances.")));
+		return;
+		// Not supported number
+	}
+
+	FString MasterMaterialPath = GetMaterial(DBlendMasterMaterial);
+
+	if (MasterMaterialPath == TEXT(""))
+	{
+		FMessageDialog::Open(EAppMsgType::Ok, FText(FText::FromString("Master Material for blending was not found. Cancelling operation.")));
+		return;
+		// Master blend material not found
+	}
+	
+	FString BlendInstanceName;
+	FString BlendDestinationPath;
+
+	if (BlendSettings->BlendedMaterialName != TEXT(""))
+	{
+		BlendInstanceName = SanitizeName(BlendSettings->BlendedMaterialName);
+	}
+
+	if (BlendSettings->BlendedMaterialPath.Path != "")
+	{
+		BlendDestinationPath = BlendSettings->BlendedMaterialPath.Path;
+	}
+
+	if (!FPackageName::IsValidObjectPath(FPaths::Combine(BlendDestinationPath, BlendInstanceName)))
+	{
+		BlendInstanceName = DBlendInstanceName;
+		BlendDestinationPath = DBlendDestinationPath;
+	}
+
+	
+	//BlendInstanceName = GetUniqueAssetName(BlendDestinationPath, BlendInstanceName, true);
+	if (UEditorAssetLibrary::DoesAssetExist(FPaths::Combine(BlendDestinationPath, BlendInstanceName)))
+	{
+		FMessageDialog::Open(EAppMsgType::Ok, FText(FText::FromString("A Blend material with same name already exists. Please choose a different name.")));
+		return;
+	}
+
+	UMaterialInstanceConstant* InstancedBlendMaterial = FMaterialUtils::CreateInstanceMaterial(MasterMaterialPath, BlendDestinationPath, BlendInstanceName);
+
+	if (InstancedBlendMaterial == nullptr)
+	{
+		FMessageDialog::Open(EAppMsgType::Ok, FText(FText::FromString("An error ocurred while creating Blend material.")));
+		return;
+	}
+
+	auto BlendSetIt = BlendSets.CreateConstIterator();
+	for (UMaterialInstanceConstant* SelectedInstance : SelectedMaterialInstances)
+	{		
+		for (FString MapType : SupportedMapTypes)
+		{
+			UTexture* PluggedMap = UMaterialEditingLibrary::GetMaterialInstanceTextureParameterValue(SelectedInstance, FName(*MapType));
+			if (PluggedMap)
+			{
+				FName TargetParameterName = *(*BlendSetIt + TEXT(" Layer ")+ MapType + TEXT(" Map"));
+				UMaterialEditingLibrary::SetMaterialInstanceTextureParameterValue(InstancedBlendMaterial, TargetParameterName, PluggedMap);
+			}
+		}
+		++BlendSetIt;
+	}
+
+}
+
+bool FMaterialBlend::ValidateSelectedAssets(TArray<FString> SelectedMaterials, FString& Failure)
+{
+	for (FString MaterialPath : SelectedMaterials)
+	{
+
+	}
+	return false;
+}
+
+
diff --git a/Engine/Plugins/MegascansPlugin/Source/MegascansPlugin/Private/Tools/BlendMaterials.h b/Engine/Plugins/MegascansPlugin/Source/MegascansPlugin/Private/Tools/BlendMaterials.h
new file mode 100644
--- /dev/null
+++ b/Engine/Plugins/MegascansPlugin/Source/MegascansPlugin/Private/Tools/BlendMaterials.h
@@ -0,0 +1,25 @@
+// Copyright Epic Games, Inc. All Rights Reserved.
+#pragma once
+#include "CoreMinimal.h"
+
+class FMaterialBlend
+{
+private:
+	FMaterialBlend() = default;
+	static TSharedPtr<FMaterialBlend> MaterialBlendInst;
+	FString DBlendDestinationPath = TEXT("/Game/BlendMaterials");
+	FString DBlendInstanceName = TEXT("BlendMaterial");
+	FString DBlendMasterMaterial = TEXT("SurfaceBlend_MasterMaterial");
+
+	//TArray<FString> BlendSets = { "Base", "R", "G", "B", "A" };
+	TArray<FString> BlendSets = { "Base", "Middle", "Top" };
+	//TArray<FString> SupportedMapTypes = { "albedo", "normal", "roughness", "specular", "gloss", "displacement", "opacity", "translucency" };
+	TArray<FString> SupportedMapTypes = { "albedo", "normal", "roughness", "displacement", "ao" };
+
+	bool ValidateSelectedAssets(TArray<FString> SelectedMaterials, FString& Failure);
+
+public:
+	static TSharedPtr<FMaterialBlend> Get();
+	void BlendSelectedMaterials();
+
+};
\ No newline at end of file
diff --git a/Engine/Plugins/MegascansPlugin/Source/MegascansPlugin/Private/UI/MSSettings.cpp b/Engine/Plugins/MegascansPlugin/Source/MegascansPlugin/Private/UI/MSSettings.cpp
new file mode 100644
--- /dev/null
+++ b/Engine/Plugins/MegascansPlugin/Source/MegascansPlugin/Private/UI/MSSettings.cpp
@@ -0,0 +1,50 @@
+// Copyright Epic Games, Inc. All Rights Reserved.
+#include "UI/MSSettings.h"
+
+
+
+UMegascansSettings::UMegascansSettings(const FObjectInitializer& ObjectInitializer)
+	: Super(ObjectInitializer) , bCreateFoliage(true), bEnableLods(true), bBatchImportPrompt(false), bEnableDisplacement(false), bApplyToSelection(false)
+
+{
+	
+}
+
+
+UMaterialBlendSettings::UMaterialBlendSettings(const FObjectInitializer& ObjectInitializer)
+	: Super(ObjectInitializer), BlendedMaterialName(TEXT("BlendMaterial"))
+
+{
+	BlendedMaterialPath.Path = TEXT("/Game/BlendMaterials");
+}
+
+UMaterialAssetSettings::UMaterialAssetSettings(const FObjectInitializer& ObjectInitializer)
+	: Super(ObjectInitializer)
+{
+
+}
+
+UMaterialPresetsSettings::UMaterialPresetsSettings(const FObjectInitializer& ObjectInitializer)
+	: Super(ObjectInitializer)
+{
+	
+}
+
+#if WITH_EDITOR
+void UMaterialPresetsSettings::PostEditChangeProperty(FPropertyChangedEvent& PropertyChangedEvent)
+{
+	Super::PostEditChangeProperty(PropertyChangedEvent);
+	// TODO : Only save the property thats getting changed.
+	UMaterialAssetSettings* MatOverridePathSettings = GetMutableDefault<UMaterialAssetSettings>();
+	MatOverridePathSettings->MasterMaterial3d = MasterMaterial3d->GetPathName();
+	MatOverridePathSettings->MasterMaterialPlant = MasterMaterialPlant->GetPathName();
+	MatOverridePathSettings->MasterMaterialSurface = MasterMaterialSurface->GetPathName();
+
+	MatOverridePathSettings->SaveConfig();
+
+	UE_LOG(LogTemp, Error, TEXT("Surface override saved : %s"), *MatOverridePathSettings->MasterMaterialSurface);
+
+}
+
+#endif
+
diff --git a/Engine/Plugins/MegascansPlugin/Source/MegascansPlugin/Private/UI/MSSettings.h b/Engine/Plugins/MegascansPlugin/Source/MegascansPlugin/Private/UI/MSSettings.h
new file mode 100644
--- /dev/null
+++ b/Engine/Plugins/MegascansPlugin/Source/MegascansPlugin/Private/UI/MSSettings.h
@@ -0,0 +1,121 @@
+// Copyright Epic Games, Inc. All Rights Reserved.
+#pragma once
+
+#include "CoreMinimal.h"
+#include "Runtime/Engine/Classes/Materials/Material.h"
+#include "UObject/ObjectMacros.h"
+#include "UObject/Object.h"
+#include "Engine/EngineTypes.h"
+#include "UObject/SoftObjectPath.h"
+#include "MSSettings.generated.h"
+
+
+UCLASS(Config = Editor)
+class UMegascansSettings
+	: public UObject
+	
+{
+	GENERATED_UCLASS_BODY()
+	//GENERATED_BODY()
+
+public:
+	/** Create foliage assets and populate Foliage Editor for 3D Plant and Scatter types. */
+	UPROPERTY(Config, DisplayName = "Auto-Populate Foliage Painter", EditAnywhere, Category = "MegascansSettings")
+		bool bCreateFoliage;
+	
+	/** Import LODs for 3D and 3D Plant assets. */
+	UPROPERTY(Config, DisplayName = "Enable LOD Setup", EditAnywhere, Category = "MegascansSettings")
+		bool bEnableLods;
+
+	/** Ask for confirmation when importing more than 10 assets. */
+	UPROPERTY(Config, DisplayName = "Prompt Before Batch Import", EditAnywhere, Category = "MegascansSettings")
+		bool bBatchImportPrompt;
+
+	/** Setup Displacement map supported default master material. */
+	UPROPERTY(Config, DisplayName = "Enable Displacement", EditAnywhere, Category = "MegascansSettings")
+		bool bEnableDisplacement;
+
+	/** Apply imported Surface on selected Actors in Editor. */
+	UPROPERTY(Config, DisplayName = "Apply to Selection", EditAnywhere, Category = "MegascansSettings")
+		bool bApplyToSelection;
+
+	/** Only import textures that are supported by the selected Master Material for the asset type. */
+	UPROPERTY(Config, DisplayName = "Import Master Material Textures", EditAnywhere, Category = "MegascansSettings")
+		bool bFilterMasterMaterialMaps;
+	
+	/** Flip Green Channel of Normal maps upon import. */
+	/*
+	UPROPERTY(Config, DisplayName = "Flip Normal Map Green Channel", EditAnywhere, Category = "MegascansSettings")
+		bool bFlipNormalGreenChannel;
+	*/
+
+
+
+};
+
+UCLASS(Config = Editor)
+class UMaterialBlendSettings
+	: public UObject	
+{
+	GENERATED_UCLASS_BODY()
+		//GENERATED_BODY()
+
+public:
+	/** Package name for Material Blend instance. */
+	UPROPERTY(Config, DisplayName = "Material Name", EditAnywhere, Category = "MaterialBlendSettings")
+		FString BlendedMaterialName;
+
+	/** Destination path for Material Blend instance. */
+	UPROPERTY(Config, DisplayName = "Destination Path", EditAnywhere, Category = "MaterialBlendSettings", meta = (ContentDir))
+		FDirectoryPath BlendedMaterialPath;
+};
+
+
+UCLASS(Config = Editor)
+class UMaterialAssetSettings
+	: public UObject
+
+{
+	GENERATED_UCLASS_BODY()
+		//GENERATED_BODY()
+
+public:
+	UPROPERTY(Config, DisplayName = "3D Master Material", EditAnywhere, Category = "MasterMaterialOverrides")
+		FString MasterMaterial3d;
+
+	UPROPERTY(Config, DisplayName = "Surface Master Material", EditAnywhere, Category = "MasterMaterialOverrides")
+		FString MasterMaterialSurface;
+
+	UPROPERTY(Config, DisplayName = "Plant Master Material", EditAnywhere, Category = "MasterMaterialOverrides")
+		FString MasterMaterialPlant;
+
+};
+
+UCLASS(Config = Editor)
+class UMaterialPresetsSettings
+	: public UObject
+	
+{
+	GENERATED_UCLASS_BODY()
+
+
+public:
+	/** Replace default master material with your own custom master material for all 3D assets. Default material is used if field is left empty. */
+	UPROPERTY(Transient, DisplayName = "3D Master Material", EditAnywhere, Category = "MasterMaterialOverrides")		
+		TLazyObjectPtr <class UMaterial> MasterMaterial3d;
+	//TLazyObjectPtr<class UMaterialInterface> MasterMaterial3d;
+
+	/** Replace default master material with your own custom master material for all Surfaces. Default material is used if field is left empty. */
+	UPROPERTY(Transient, DisplayName = "Surface Master Material", EditAnywhere, Category = "MasterMaterialOverrides")
+		TLazyObjectPtr<class UMaterial> MasterMaterialSurface;
+
+	/** Replace default master material with your own custom master material for all 3D Plants. Default material is used if field is left empty. */
+	UPROPERTY(Transient, DisplayName = "Plant Master Material", EditAnywhere, Category = "MasterMaterialOverrides")
+		TLazyObjectPtr<class UMaterial> MasterMaterialPlant;
+
+#if WITH_EDITOR
+	virtual void PostEditChangeProperty(FPropertyChangedEvent& PropertyThatWillChange) override;
+#endif
+
+};
+
diff --git a/Engine/Plugins/MegascansPlugin/Source/MegascansPlugin/Private/UI/MSStyle.cpp b/Engine/Plugins/MegascansPlugin/Source/MegascansPlugin/Private/UI/MSStyle.cpp
new file mode 100644
--- /dev/null
+++ b/Engine/Plugins/MegascansPlugin/Source/MegascansPlugin/Private/UI/MSStyle.cpp
@@ -0,0 +1,72 @@
+// Copyright Epic Games, Inc. All Rights Reserved.
+#include "UI/MSStyle.h"
+
+#include "Framework/Application/SlateApplication.h"
+#include "Styling/SlateStyle.h"
+#include "Styling/SlateStyleRegistry.h"
+
+TUniquePtr<FSlateStyleSet> FMSStyle::MSStyleInstance;
+
+void FMSStyle::Initialize()
+{
+	if (!MSStyleInstance.IsValid())
+	{
+		MSStyleInstance = Create();
+		FSlateStyleRegistry::RegisterSlateStyle(*MSStyleInstance);
+	}
+}
+
+void FMSStyle::Shutdown()
+{
+	if (MSStyleInstance.IsValid())
+	{
+		
+		FSlateStyleRegistry::UnRegisterSlateStyle(*MSStyleInstance);
+		MSStyleInstance.Reset();
+	}
+}
+
+FName FMSStyle::GetStyleSetName()
+{
+	static FName StyleSetName(TEXT("MegascansStyle"));
+	return StyleSetName;
+}
+
+FName FMSStyle::GetContextName()
+{
+	//static FName ContextName(TEXT("Megascans"));
+	return FName(TEXT("Megascans"));
+}
+
+#define IMAGE_BRUSH( RelativePath, ... ) FSlateImageBrush( Style->RootToContentDir( RelativePath, TEXT(".png") ), __VA_ARGS__ )
+const FVector2D Icon20x20(20.0f, 20.0f);
+const FVector2D Icon40x40(40.0f, 40.0f);
+
+TUniquePtr< FSlateStyleSet > FMSStyle::Create()
+{
+	TUniquePtr< FSlateStyleSet > Style = MakeUnique<FSlateStyleSet>(GetStyleSetName());
+	Style->SetContentRoot(FPaths::EnginePluginsDir() / TEXT("MegascansPlugin/Resources"));
+	return Style;
+}
+
+void FMSStyle::SetIcon(const FString& StyleName, const FString& ResourcePath)
+{
+	FSlateStyleSet* Style = MSStyleInstance.Get();
+
+	FString Name(GetContextName().ToString());
+	Name = Name + "." + StyleName;
+	Style->Set(*Name, new IMAGE_BRUSH(ResourcePath, Icon40x40));
+
+	Name += ".Small";
+	Style->Set(*Name, new IMAGE_BRUSH(ResourcePath, Icon20x20));
+
+	FSlateApplication::Get().GetRenderer()->ReloadTextureResources();
+}
+
+#undef IMAGE_BRUSH
+
+const ISlateStyle& FMSStyle::Get()
+{
+	check(MSStyleInstance);
+	return *MSStyleInstance;
+}
diff --git a/Engine/Plugins/MegascansPlugin/Source/MegascansPlugin/Private/UI/MSStyle.h b/Engine/Plugins/MegascansPlugin/Source/MegascansPlugin/Private/UI/MSStyle.h
new file mode 100644
--- /dev/null
+++ b/Engine/Plugins/MegascansPlugin/Source/MegascansPlugin/Private/UI/MSStyle.h
@@ -0,0 +1,23 @@
+// Copyright Epic Games, Inc. All Rights Reserved.
+#pragma once
+#include "CoreMinimal.h"
+
+
+class FSlateStyleSet;
+class ISlateStyle;
+class FMSStyle
+{
+public:
+	static void Initialize();
+	static void Shutdown();
+
+	static const ISlateStyle& Get();
+	static FName GetStyleSetName();
+	static FName GetContextName();
+
+	static void SetIcon(const FString& StyleName, const FString& ResourcePath);
+
+private:
+	static TUniquePtr<FSlateStyleSet> Create();
+	static TUniquePtr<FSlateStyleSet> MSStyleInstance;
+};
\ No newline at end of file
diff --git a/Engine/Plugins/MegascansPlugin/Source/MegascansPlugin/Private/UI/MaterialBlendingDetails.cpp b/Engine/Plugins/MegascansPlugin/Source/MegascansPlugin/Private/UI/MaterialBlendingDetails.cpp
new file mode 100644
--- /dev/null
+++ b/Engine/Plugins/MegascansPlugin/Source/MegascansPlugin/Private/UI/MaterialBlendingDetails.cpp
@@ -0,0 +1,56 @@
+// Copyright Epic Games, Inc. All Rights Reserved.
+#include "MaterialBlendingDetails.h"
+#include "UI/MSSettings.h"
+
+#include "DetailCategoryBuilder.h"
+#include "DetailLayoutBuilder.h"
+#include "DetailWidgetRow.h"
+
+#include "Tools/BlendMaterials.h"
+
+#define LOCTEXT_NAMESPACE "MegascansLiveLinkModule"
+
+void BlendSettingsCustomization::CustomizeDetails(IDetailLayoutBuilder& DetailBuilder)
+{
+	IDetailCategoryBuilder& BlendCategory = DetailBuilder.EditCategory("MaterialBlendSettings");
+
+	FDetailWidgetRow BlendButton = BlendCategory.AddCustomRow(NSLOCTEXT("MaterialBlendSettings", "CreateMaterialBlend", "Create Material Blend"))
+		[
+			SNew(SHorizontalBox)
+			+ SHorizontalBox::Slot()
+		.FillWidth(1.f)
+		.Padding(0, 5, 10, 5)
+		[
+			SNew(SButton)
+			.ContentPadding(3)
+		.VAlign(VAlign_Center)
+		.HAlign(HAlign_Center)
+		.IsEnabled(this, &BlendSettingsCustomization::CanCreateMaterialBlend)
+		.OnClicked(this, &BlendSettingsCustomization::CreateMaterialBlend)
+		.Text(NSLOCTEXT("MaterialBlendSettings", "CreateMaterialBlend", "Create Material Blend"))
+		]
+		];
+
+	//BlendButton.NameContent().VAlign(EVerticalAlignment::VAlign_Center);
+
+}
+
+TSharedRef<IDetailCustomization> BlendSettingsCustomization::MakeInstance()
+{
+	return MakeShareable(new BlendSettingsCustomization);
+}
+
+bool BlendSettingsCustomization::CanCreateMaterialBlend() const
+{
+	return true;
+}
+
+FReply BlendSettingsCustomization::CreateMaterialBlend()
+{
+	FMaterialBlend::Get()->BlendSelectedMaterials();
+	return FReply::Handled();
+}
+
+
+#undef LOCTEXT_NAMESPACE
+
diff --git a/Engine/Plugins/MegascansPlugin/Source/MegascansPlugin/Private/UI/MaterialBlendingDetails.h b/Engine/Plugins/MegascansPlugin/Source/MegascansPlugin/Private/UI/MaterialBlendingDetails.h
new file mode 100644
--- /dev/null
+++ b/Engine/Plugins/MegascansPlugin/Source/MegascansPlugin/Private/UI/MaterialBlendingDetails.h
@@ -0,0 +1,17 @@
+// Copyright Epic Games, Inc. All Rights Reserved.
+#pragma once
+#include "IDetailCustomization.h"
+
+class BlendSettingsCustomization : public IDetailCustomization
+
+{
+public:
+	virtual void CustomizeDetails(IDetailLayoutBuilder& DetailBuilder) override;
+
+	static TSharedRef<IDetailCustomization> MakeInstance();
+
+private:
+	bool CanCreateMaterialBlend() const;
+	FReply CreateMaterialBlend();
+
+};
\ No newline at end of file
diff --git a/Engine/Plugins/MegascansPlugin/Source/MegascansPlugin/Private/UI/QMSUIManager.cpp b/Engine/Plugins/MegascansPlugin/Source/MegascansPlugin/Private/UI/QMSUIManager.cpp
new file mode 100644
--- /dev/null
+++ b/Engine/Plugins/MegascansPlugin/Source/MegascansPlugin/Private/UI/QMSUIManager.cpp
@@ -0,0 +1,138 @@
+// Copyright Epic Games, Inc. All Rights Reserved.
+#include "UI/QMSUIManager.h"
+#include "UI/MSStyle.h"
+
+#include "AssetData.h"
+#include "ContentBrowserModule.h"
+#include "LevelEditor.h"
+#include "Framework/MultiBox/MultiBoxBuilder.h"
+#include "Framework/MultiBox/MultiBoxExtender.h"
+#include "Misc/Paths.h"
+#include "Modules/ModuleManager.h"
+
+#include "IPythonScriptPlugin.h"
+
+#include "MSPythonBridge.h"
+#include "TCPServerImp.h"
+#include "UI/SMSWindow.h"
+
+
+#define LOCTEXT_NAMESPACE "QMSLiveLink"
+#define LEVELEDITOR_MODULE_NAME TEXT("LevelEditor")
+#define CONTENTBROWSER_MODULE_NAME TEXT("ContentBrowser")
+
+TUniquePtr<FQMSUIManagerImpl> FQMSUIManager::Instance;
+
+class FQMSUIManagerImpl
+{
+public:
+	void Initialize();
+	void Shutdown();
+	void JsonReceived(FString message);
+
+private:
+	void SetupMenuItem();
+	void CreateWIndow();
+	FTCPServer *tcpServer = NULL;	
+	void FillToolbar(FToolBarBuilder& ToolbarBuilder);
+	void AddPluginMenu(FMenuBuilder& MenuBuilder);
+};
+
+void FQMSUIManager::Initialize()
+{
+	if (!Instance.IsValid())
+	{
+		Instance = MakeUnique<FQMSUIManagerImpl>();
+		Instance->Initialize();
+	}	
+}
+
+void FQMSUIManagerImpl::Initialize()
+{
+	FMSStyle::Initialize();
+
+	SetupMenuItem();
+	if (tcpServer == NULL)
+	{
+		tcpServer = new FTCPServer();
+	}
+
+	
+}
+
+void FQMSUIManagerImpl::SetupMenuItem() 
+{
+	FMSStyle::SetIcon("Logo", "Logo80x80");
+	FMSStyle::SetIcon("ContextLogo", "Logo32x32");
+
+	FLevelEditorModule& LevelEditorModule = FModuleManager::LoadModuleChecked<FLevelEditorModule>(LEVELEDITOR_MODULE_NAME);
+	//TSharedPtr<FExtender> ToolbarExtender = MakeShareable(new FExtender);
+	//ToolbarExtender->AddToolBarExtension("Settings", EExtensionHook::After, nullptr, FToolBarExtensionDelegate::CreateRaw(this, &FQMSUIManagerImpl::FillToolbar));
+	//LevelEditorModule.GetToolBarExtensibilityManager()->AddExtender(ToolbarExtender);
+
+	TSharedPtr<FExtender> NewMenuExtender = MakeShareable(new FExtender);
+	NewMenuExtender->AddMenuExtension("LevelEditor",
+		EExtensionHook::After,
+		NULL,
+		FMenuExtensionDelegate::CreateRaw(this, &FQMSUIManagerImpl::AddPluginMenu));
+
+	LevelEditorModule.GetMenuExtensibilityManager()->AddExtender(NewMenuExtender);
+
+}
+
+
+void FQMSUIManagerImpl::AddPluginMenu(FMenuBuilder& MenuBuilder)
+{
+	MenuBuilder.BeginSection("CustomMenu", TAttribute<FText>(FText::FromString("Megascans")));
+
+	MenuBuilder.AddMenuEntry(
+		LOCTEXT("OpenWindow", "Megascans"),
+		LOCTEXT("ToolTip", "Open Megascans import settings"),
+		FSlateIcon(),
+		FUIAction(FExecuteAction::CreateRaw(this, &FQMSUIManagerImpl::CreateWIndow))
+	);
+
+	MenuBuilder.EndSection();
+}
+
+void FQMSUIManagerImpl::FillToolbar(FToolBarBuilder& ToolbarBuilder)
+{
+	ToolbarBuilder.BeginSection(TEXT("QuixelMSLiveLink"));
+	{
+			ToolbarBuilder.AddToolBarButton(
+				FUIAction(FExecuteAction::CreateRaw( this, &FQMSUIManagerImpl::CreateWIndow)),
+				FName(TEXT("Quixel MS LiveLink")),
+				LOCTEXT("QMSLiveLink_label", "Megascans"),
+				LOCTEXT("WorldProperties_ToolTipOverride", "Megascans Link with Bridge"),
+				FSlateIcon(FMSStyle::GetStyleSetName(), "Megascans.Logo"),
+				EUserInterfaceActionType::Button,
+				FName(TEXT("QuixelMSLiveLink"))
+				);
+	}
+	ToolbarBuilder.EndSection();
+}
+
+void FQMSUIManagerImpl::CreateWIndow() 
+{
+	FLevelEditorModule& LevelEditorModule = FModuleManager::GetModuleChecked<FLevelEditorModule>(TEXT("LevelEditor"));
+	MegascansSettingsWindow::OpenSettingsWindow(LevelEditorModule.GetLevelEditorTabManager().ToSharedRef());
+	/*
+	UMSPythonBridge* bridge = UMSPythonBridge::Get();
+	bridge->InitializePythonWindow();
+	bridge->TestFbxImport();
+	*/
+}
+
+void FQMSUIManagerImpl::JsonReceived(FString message) {
+	UMSPythonBridge* bridge = UMSPythonBridge::Get();
+	bridge->GetUeData(message);	
+}
+
+void FQMSUIManager::Shutdown()
+{
+	// Do other destruction stuff
+	FMSStyle::Shutdown();
+}
+
+#undef LOCTEXT_NAMESPACE
+
diff --git a/Engine/Plugins/MegascansPlugin/Source/MegascansPlugin/Private/UI/QMSUIManager.h b/Engine/Plugins/MegascansPlugin/Source/MegascansPlugin/Private/UI/QMSUIManager.h
new file mode 100644
--- /dev/null
+++ b/Engine/Plugins/MegascansPlugin/Source/MegascansPlugin/Private/UI/QMSUIManager.h
@@ -0,0 +1,23 @@
+// Copyright Epic Games, Inc. All Rights Reserved.
+#pragma once
+
+#include "CoreMinimal.h"
+
+
+
+
+class FQMSUIManagerImpl;
+
+
+
+class  FQMSUIManager
+{
+public:
+	static void Initialize();
+	static void Shutdown();
+	static TUniquePtr<FQMSUIManagerImpl> Instance;
+
+//private:
+	//static TUniquePtr<FQMSUIManagerImpl> Instance;
+};
+
diff --git a/Engine/Plugins/MegascansPlugin/Source/MegascansPlugin/Private/UI/SMSWindow.cpp b/Engine/Plugins/MegascansPlugin/Source/MegascansPlugin/Private/UI/SMSWindow.cpp
new file mode 100644
--- /dev/null
+++ b/Engine/Plugins/MegascansPlugin/Source/MegascansPlugin/Private/UI/SMSWindow.cpp
@@ -0,0 +1,258 @@
+
+// Copyright Epic Games, Inc. All Rights Reserved.
+#include "UI/SMSWindow.h"
+
+
+#include "Modules/ModuleManager.h"
+#include "UObject/Class.h"
+#include "UObject/UObjectHash.h"
+#include "UObject/UObjectIterator.h"
+#include "UObject/GCObject.h"
+#include "Layout/Margin.h"
+#include "Input/Reply.h"
+#include "Widgets/DeclarativeSyntaxSupport.h"
+#include "Widgets/SCompoundWidget.h"
+#include "Widgets/SBoxPanel.h"
+#include "Widgets/SWindow.h"
+#include "Framework/Application/SlateApplication.h"
+#include "Widgets/Text/STextBlock.h"
+#include "Widgets/Input/SEditableTextBox.h"
+#include "Widgets/Input/SButton.h"
+#include "EditorStyleSet.h"
+#include "Framework/Docking/TabManager.h"
+//#include "Toolkits/AssetEditorManager.h"
+#include "Widgets/Layout/SScrollBox.h"
+#include "IDetailsView.h"
+#include "PropertyEditorModule.h"
+#include "Widgets/Docking/SDockTab.h"
+
+#include "Runtime/Engine/Classes/Materials/MaterialInterface.h"
+#include "EditorAssetLibrary.h"
+
+
+#define LOCTEXT_NAMESPACE "MegascansSettings"
+
+TWeakPtr<SWindow> MSSettingsWindow;
+
+
+class SMegascansSettings : public SCompoundWidget, public FGCObject
+{
+	SLATE_BEGIN_ARGS(SMegascansSettings) {}
+		SLATE_ARGUMENT(UMegascansSettings*, InMegascansSettings)
+		SLATE_ARGUMENT(UMaterialBlendSettings*, InMaterialBlendSettings)
+		SLATE_ARGUMENT(UMaterialPresetsSettings*, InMaterialOverrideSettings)
+
+	SLATE_END_ARGS()
+
+	void Construct(const FArguments& InArgs)
+	{
+		//const UMegascansSettings* MegascansSettings = GetDefault<UMegascansSettings>();
+		
+
+		FPropertyEditorModule& PropertyEditor = FModuleManager::LoadModuleChecked<FPropertyEditorModule>("PropertyEditor");
+
+		FDetailsViewArgs MiscOptionsDetailsArgs;
+		MiscOptionsDetailsArgs.bUpdatesFromSelection = false;
+		MiscOptionsDetailsArgs.bLockable = false;
+		MiscOptionsDetailsArgs.bAllowSearch = false;
+		MiscOptionsDetailsArgs.bShowOptions = false;
+		MiscOptionsDetailsArgs.bAllowFavoriteSystem = false;
+		MiscOptionsDetailsArgs.NameAreaSettings = FDetailsViewArgs::HideNameArea;
+		MiscOptionsDetailsArgs.ViewIdentifier = "MegascansSettings";
+
+		MiscOptionsDetails = PropertyEditor.CreateDetailView(MiscOptionsDetailsArgs);
+
+		FDetailsViewArgs MatBlendDetailsArgs;
+		MatBlendDetailsArgs.bUpdatesFromSelection = false;
+		MatBlendDetailsArgs.bLockable = false;
+		MatBlendDetailsArgs.bAllowSearch = false;
+		MatBlendDetailsArgs.bShowOptions = false;
+		MatBlendDetailsArgs.bAllowFavoriteSystem = false;
+		MatBlendDetailsArgs.NameAreaSettings = FDetailsViewArgs::HideNameArea;
+		MatBlendDetailsArgs.ViewIdentifier = "MaterialBlendingSettings";
+
+		MatBlendDetails = PropertyEditor.CreateDetailView(MatBlendDetailsArgs);
+
+		FDetailsViewArgs MaterialOverrideArgs;
+		MaterialOverrideArgs.bUpdatesFromSelection = false;
+		MaterialOverrideArgs.bLockable = false;
+		MaterialOverrideArgs.bAllowSearch = false;
+		MaterialOverrideArgs.bShowOptions = false;
+		MaterialOverrideArgs.bAllowFavoriteSystem = false;
+		MaterialOverrideArgs.NameAreaSettings = FDetailsViewArgs::HideNameArea;
+		MaterialOverrideArgs.ViewIdentifier = "MasterMaterialOverrides";
+
+		MatOverrideDetails = PropertyEditor.CreateDetailView(MaterialOverrideArgs);
+
+		
+
+		ChildSlot
+		[
+			SNew(SVerticalBox)
+
+			+ SVerticalBox::Slot()
+			.FillHeight(1.0f)
+			.Padding(4, 4, 4, 4)
+			[
+				SNew(SScrollBox)
+				+SScrollBox::Slot()
+				[
+					MiscOptionsDetails.ToSharedRef()
+				]
+
+				+SScrollBox::Slot()
+					[
+						MatOverrideDetails.ToSharedRef()
+					]
+
+				+ SScrollBox::Slot()
+					[
+						MatBlendDetails.ToSharedRef()
+						
+					]
+
+			]
+
+		];		
+
+		if (InArgs._InMegascansSettings)
+		{
+			SetMegascansSettings(InArgs._InMegascansSettings);
+		}
+
+		if (InArgs._InMaterialBlendSettings)
+		{
+			SetBlendSettings(InArgs._InMaterialBlendSettings);
+		}
+
+		if (InArgs._InMaterialOverrideSettings)
+		{
+			SetMaterialOverrideSettings(InArgs._InMaterialOverrideSettings);
+		}
+		
+	}
+
+
+
+	void SetMegascansSettings(UMegascansSettings* InMegascansSettings)
+	{
+		MegascansSettings = InMegascansSettings;
+
+		MiscOptionsDetails->SetObject(MegascansSettings);
+	}
+
+	void SetBlendSettings(UMaterialBlendSettings* InBlendSettings)
+	{
+		MaterialBlendSettings = InBlendSettings;
+		MatBlendDetails->SetObject(MaterialBlendSettings);
+
+	}
+
+	void SetMaterialOverrideSettings(UMaterialPresetsSettings* InMaterialOverrideSettings)
+	{
+		MaterialOverrideSettings = InMaterialOverrideSettings;
+		MatOverrideDetails->SetObject(MaterialOverrideSettings);
+
+	}
+
+	virtual void AddReferencedObjects( FReferenceCollector& Collector ) override
+	{
+		Collector.AddReferencedObject(MegascansSettings);
+		Collector.AddReferencedObject(MaterialBlendSettings);
+		Collector.AddReferencedObject(MaterialOverrideSettings);
+	}
+
+private:
+
+	TSharedPtr<IDetailsView> MiscOptionsDetails;
+	TSharedPtr<IDetailsView> MatBlendDetails;
+	TSharedPtr<IDetailsView> MatOverrideDetails;
+
+	UMegascansSettings* MegascansSettings;
+	UMaterialBlendSettings* MaterialBlendSettings;
+	UMaterialPresetsSettings* MaterialOverrideSettings;
+};
+	
+
+
+
+
+void MegascansSettingsWindow::OpenSettingsWindow(const TSharedRef<FTabManager>& TabManager)
+{
+	TSharedPtr<SWindow> ExistingWindow = MSSettingsWindow.Pin();
+	if (ExistingWindow.IsValid())
+	{
+		ExistingWindow->BringToFront();
+	}
+	else
+	{
+		ExistingWindow = SNew(SWindow)
+			.Title( LOCTEXT("MegascansSettings", "Megascans") )
+			.HasCloseButton(true)
+			.SupportsMaximize(false)
+			.SupportsMinimize(false)
+			
+			.ClientSize(FVector2D(600, 600));
+
+		TSharedPtr<SDockTab> OwnerTab = TabManager->GetOwnerTab();
+		TSharedPtr<SWindow> RootWindow = OwnerTab.IsValid() ? OwnerTab->GetParentWindow() : TSharedPtr<SWindow>();
+		if(RootWindow.IsValid())
+		{
+			FSlateApplication::Get().AddWindowAsNativeChild(ExistingWindow.ToSharedRef(), RootWindow.ToSharedRef());
+		}
+		else
+		{
+			FSlateApplication::Get().AddWindow(ExistingWindow.ToSharedRef());
+		}
+	}
+
+	UMegascansSettings* MegascansSettings = GetMutableDefault<UMegascansSettings>();
+	UMaterialBlendSettings* MaterialBlendSettings = GetMutableDefault<UMaterialBlendSettings>();
+	UMaterialPresetsSettings* MaterialOverrideSettings = GetMutableDefault<UMaterialPresetsSettings>();
+
+	//Settings override changing
+	const UMaterialAssetSettings* MatAssetPathsSettings = GetDefault<UMaterialAssetSettings>();
+	if(MatAssetPathsSettings->MasterMaterial3d != "None" && MatAssetPathsSettings->MasterMaterial3d != "")
+		if (UEditorAssetLibrary::DoesAssetExist(MatAssetPathsSettings->MasterMaterial3d)) MaterialOverrideSettings->MasterMaterial3d = CastChecked<UMaterial>(UEditorAssetLibrary::LoadAsset(MatAssetPathsSettings->MasterMaterial3d));
+	if (MatAssetPathsSettings->MasterMaterialPlant != "None" && MatAssetPathsSettings->MasterMaterialPlant != "")
+		if (UEditorAssetLibrary::DoesAssetExist(MatAssetPathsSettings->MasterMaterialPlant)) MaterialOverrideSettings->MasterMaterialPlant = CastChecked<UMaterial>(UEditorAssetLibrary::LoadAsset(MatAssetPathsSettings->MasterMaterialPlant));
+	if (MatAssetPathsSettings->MasterMaterialSurface != "None" && MatAssetPathsSettings->MasterMaterialSurface != "")
+		if (UEditorAssetLibrary::DoesAssetExist(MatAssetPathsSettings->MasterMaterialSurface)) MaterialOverrideSettings->MasterMaterialSurface = CastChecked<UMaterial>(UEditorAssetLibrary::LoadAsset(MatAssetPathsSettings->MasterMaterialSurface));
+
+	ExistingWindow->SetContent(
+		SNew(SMegascansSettings)
+		.InMegascansSettings(MegascansSettings)
+		.InMaterialBlendSettings(MaterialBlendSettings)
+		.InMaterialOverrideSettings(MaterialOverrideSettings)
+
+	);
+	
+	
+	//ExistingWindow->SetOnWindowClosed(FOnWindowClosed::CreateStatic(&MegascansSettingsWindow::SaveSettings, MegascansSettings));	
+	
+	MSSettingsWindow = ExistingWindow;
+	
+}
+
+void MegascansSettingsWindow::SaveSettings(const TSharedRef<SWindow>& Window, UMegascansSettings* MegascansSettings)
+{	
+	MegascansSettings->SaveConfig();
+
+	//Save Master material overrides.
+	
+	const UMaterialPresetsSettings* MatPresetsSettings = GetDefault<UMaterialPresetsSettings>();
+	UMaterialAssetSettings* MatOverridePathSettings = GetMutableDefault<UMaterialAssetSettings>();
+
+	MatOverridePathSettings->MasterMaterial3d = MatPresetsSettings->MasterMaterial3d->GetPathName();
+	MatOverridePathSettings->MasterMaterialPlant = MatPresetsSettings->MasterMaterialPlant->GetPathName();
+	MatOverridePathSettings->MasterMaterialSurface = MatPresetsSettings->MasterMaterialSurface->GetPathName();
+	MatOverridePathSettings->SaveConfig();
+	
+
+}
+
+
+
+
+
+#undef LOCTEXT_NAMESPACE
\ No newline at end of file
diff --git a/Engine/Plugins/MegascansPlugin/Source/MegascansPlugin/Private/UI/SMSWindow.h b/Engine/Plugins/MegascansPlugin/Source/MegascansPlugin/Private/UI/SMSWindow.h
new file mode 100644
--- /dev/null
+++ b/Engine/Plugins/MegascansPlugin/Source/MegascansPlugin/Private/UI/SMSWindow.h
@@ -0,0 +1,23 @@
+// Copyright Epic Games, Inc. All Rights Reserved.
+#pragma once
+
+
+#include "CoreMinimal.h"
+#include "UI/MSSettings.h"
+
+class FTabManager;
+
+class MegascansSettingsWindow
+{
+public:
+	
+	
+	static void OpenSettingsWindow(const TSharedRef<FTabManager>& TabManager);
+	static void SaveSettings(const TSharedRef<SWindow>& Window, UMegascansSettings* MegascansSettings);	
+	
+	
+};
+
+
+
+
diff --git a/Engine/Plugins/MegascansPlugin/Source/MegascansPlugin/Private/Utilities/Analytics.cpp b/Engine/Plugins/MegascansPlugin/Source/MegascansPlugin/Private/Utilities/Analytics.cpp
new file mode 100644
--- /dev/null
+++ b/Engine/Plugins/MegascansPlugin/Source/MegascansPlugin/Private/Utilities/Analytics.cpp
@@ -0,0 +1,99 @@
+// Copyright Epic Games, Inc. All Rights Reserved.
+#include "Utilities/Analytics.h"
+#include "UI/MSSettings.h"
+#include "Interfaces/IHttpRequest.h"
+
+#include "Runtime/Json/Public/Dom/JsonObject.h"
+#include "Runtime/Json/Public/Serialization/JsonSerializer.h"
+#include "HttpModule.h"
+#include "Misc/EngineVersion.h"
+
+TSharedPtr<FAnalytics> FAnalytics::AnalyticsInst;
+
+TSharedPtr<FAnalytics> FAnalytics::Get()
+{
+	if (!AnalyticsInst.IsValid())
+	{
+		AnalyticsInst = MakeShareable(new FAnalytics);
+	}
+	return AnalyticsInst;
+}
+
+FString FAnalytics::GetUEVersion()
+{
+	FString EngineVersion = TEXT("None");
+#if ENGINE_MINOR_VERSION == 23 
+	EngineVersion = TEXT("4.23");
+#elif ENGINE_MINOR_VERSION == 24
+	EngineVersion = TEXT("4.24");
+#elif ENGINE_MINOR_VERSION == 25
+	EngineVersion = TEXT("4.25");
+#elif ENGINE_MINOR_VERSION == 26
+	EngineVersion = TEXT("4.26");
+#elif ENGINE_MINOR_VERSION == 27
+	EngineVersion = TEXT("4.27");
+#elif ENGINE_MINOR_VERSION == 28
+	EngineVersion = TEXT("4.28");
+#endif
+	return EngineVersion;
+}
+
+void FAnalytics::SendAnalytics(TSharedPtr<FJsonObject> JsonObject)
+{
+	//http://localhost:28241/analytics
+	FString UrlAddressAsString = TEXT("https://stats.quixel.com/v1/message");	
+    FString OutputString;
+    TSharedRef<TJsonWriter<TCHAR>> JsonWriter = TJsonWriterFactory<>::Create(&OutputString);
+    FJsonSerializer::Serialize(JsonObject.ToSharedRef(), JsonWriter);
+    //TSharedRef<IHttpRequest> HttpRequest = FHttpModule::Get().CreateRequest();
+	TSharedRef<IHttpRequest, ESPMode::ThreadSafe> HttpRequest = FHttpModule::Get().CreateRequest();
+    HttpRequest->SetVerb(TEXT("POST"));
+    HttpRequest->SetHeader(TEXT("Content-Type"), TEXT("application/json"));
+	HttpRequest->SetHeader(TEXT("User-Agent"), "X-UnrealEngine-Agent");
+    HttpRequest->SetURL(*FString::Printf(TEXT("%s"), *UrlAddressAsString));
+    HttpRequest->SetContentAsString(OutputString);
+	//UE_LOG(LogTemp, Error, TEXT("Selected asset type is : %s"), *OutputString);
+    //HttpRequest->OnProcessRequestComplete().BindUObject(this, &FAnalytics::ResponseReceived);
+	
+    HttpRequest->ProcessRequest();
+}
+
+TSharedPtr<FJsonObject> FAnalytics::GenerateAnalyticsJson()
+{
+	TSharedPtr<FJsonObject> JsonObject = MakeShareable(new FJsonObject);
+	TSharedPtr<FJsonObject> AnalyticsJson = MakeShareable(new FJsonObject);	
+
+    const UMegascansSettings* MegascansSettings = GetDefault<UMegascansSettings>();
+	const UMaterialPresetsSettings* MatOverrideSettings = GetDefault< UMaterialPresetsSettings>();
+
+	JsonObject->SetStringField(TEXT("UEVersion"), GetUEVersion());
+    JsonObject->SetStringField(TEXT("CreateFoliage"), MegascansSettings->bCreateFoliage ? *FString::Printf(TEXT("True")) : *FString::Printf(TEXT("False")));
+    JsonObject->SetStringField(TEXT("EnableLODS"), MegascansSettings->bEnableLods ? *FString::Printf(TEXT("True")) : *FString::Printf(TEXT("False")));
+    JsonObject->SetStringField(TEXT("BatchPrompt"), MegascansSettings->bBatchImportPrompt ? *FString::Printf(TEXT("True")) : *FString::Printf(TEXT("False")));
+    JsonObject->SetStringField(TEXT("EnableDisplacement"), MegascansSettings->bEnableDisplacement ? *FString::Printf(TEXT("True")) : *FString::Printf(TEXT("False")));
+    JsonObject->SetStringField(TEXT("ApplyToSelection"), MegascansSettings->bApplyToSelection ? *FString::Printf(TEXT("True")) : *FString::Printf(TEXT("False")));
+    JsonObject->SetStringField(TEXT("FilterMaps"), MegascansSettings->bFilterMasterMaterialMaps ? *FString::Printf(TEXT("True")) : *FString::Printf(TEXT("False")));    
+    JsonObject->SetStringField(TEXT("MaterialOverride_3D"), ( MatOverrideSettings->MasterMaterial3d != nullptr) ? *FString::Printf(TEXT("True")) : *FString::Printf(TEXT("False")));
+    JsonObject->SetStringField(TEXT("MaterialOverride_Surface"), ( MatOverrideSettings->MasterMaterialSurface != nullptr) ? *FString::Printf(TEXT("True")) : *FString::Printf(TEXT("False")));
+    JsonObject->SetStringField(TEXT("MaterialOverride_3DPlant"), ( MatOverrideSettings->MasterMaterialPlant != nullptr) ? *FString::Printf(TEXT("True")) : *FString::Printf(TEXT("False")));
+	
+	AnalyticsJson->SetObjectField(TEXT("data"), JsonObject);
+	AnalyticsJson->SetStringField(TEXT("event"), TEXT("Unreal_Plugin_Analytics"));
+	return AnalyticsJson;
+}
+
+TSharedPtr<FJsonObject> FAnalytics::GenerateBlendAnalyticsJson()
+{
+	TSharedPtr<FJsonObject> JsonObject = MakeShareable(new FJsonObject);
+	TSharedPtr<FJsonObject> AnalyticsJson = MakeShareable(new FJsonObject);
+    const UMaterialBlendSettings* BlendSettings = GetDefault<UMaterialBlendSettings>();
+
+	JsonObject->SetStringField(TEXT("UEVersion"), GetUEVersion());
+    JsonObject->SetStringField(TEXT("BlendMaterialUsed"), *FString::Printf(TEXT("True")) );
+
+	AnalyticsJson->SetObjectField(TEXT("data"), JsonObject);
+	AnalyticsJson->SetStringField(TEXT("event"), TEXT("Unreal_Plugin_Analytics"));
+	return AnalyticsJson;
+    
+    
+}
\ No newline at end of file
diff --git a/Engine/Plugins/MegascansPlugin/Source/MegascansPlugin/Private/Utilities/Analytics.h b/Engine/Plugins/MegascansPlugin/Source/MegascansPlugin/Private/Utilities/Analytics.h
new file mode 100644
--- /dev/null
+++ b/Engine/Plugins/MegascansPlugin/Source/MegascansPlugin/Private/Utilities/Analytics.h
@@ -0,0 +1,24 @@
+// Copyright Epic Games, Inc. All Rights Reserved.
+#pragma once
+#include "CoreMinimal.h"
+
+
+
+
+class FAnalytics
+{
+private:
+	FAnalytics() = default;
+	static TSharedPtr<FAnalytics> AnalyticsInst;
+	FString GetUEVersion();
+
+    
+
+
+public:
+	static TSharedPtr<FAnalytics> Get();
+    void SendAnalytics(TSharedPtr<FJsonObject> JsonObject);
+	TSharedPtr<FJsonObject> GenerateAnalyticsJson();
+	TSharedPtr<FJsonObject> GenerateBlendAnalyticsJson( );
+	
+};
\ No newline at end of file
diff --git a/Engine/Plugins/MegascansPlugin/Source/MegascansPlugin/Private/Utilities/AssetData.cpp b/Engine/Plugins/MegascansPlugin/Source/MegascansPlugin/Private/Utilities/AssetData.cpp
new file mode 100644
--- /dev/null
+++ b/Engine/Plugins/MegascansPlugin/Source/MegascansPlugin/Private/Utilities/AssetData.cpp
@@ -0,0 +1,232 @@
+// Copyright Epic Games, Inc. All Rights Reserved.
+#include "AssetData.h"
+#include "Utilities/MiscUtils.h"
+#include "UI/MSSettings.h"
+#include "EditorAssetLibrary.h"
+
+
+TSharedPtr<FAssetImportParams> FAssetImportParams::ImportParamsInst;
+TSharedPtr<FAssetImportParams> FAssetImportParams::Get()
+{
+	if (!ImportParamsInst.IsValid())
+	{
+		ImportParamsInst = MakeShareable(new FAssetImportParams);
+	}
+	return ImportParamsInst;
+}
+
+
+TSharedPtr<ImportParams3DAsset> FAssetImportParams::Get3DAssetsParams(TSharedPtr<FAssetTypeData> AssetImportData)
+{
+	TSharedPtr<ImportParams3DAsset> AssetParams = MakeShareable(new ImportParams3DAsset);
+	AssetParams->BaseParams = GetAssetImportParams(AssetImportData);
+
+
+	const UMaterialAssetSettings* MatAssetPathsSettings = GetDefault<UMaterialAssetSettings>();
+	TSharedPtr<Asset3DParams>  Asset3DParameters = MakeShareable(new Asset3DParams);
+	Asset3DParameters->SubType = Get3DSubtype(AssetImportData);
+	Asset3DParameters->MeshDestination = AssetParams->BaseParams->AssetDestination;
+	Asset3DParameters->Asset3DMaterialType = GetS3DMasterMaterialType(AssetImportData);
+	
+	//TODO : create mapping of material with texture sets
+	TSharedPtr<SurfaceParams> MaterialParams = MakeShareable(new SurfaceParams);
+	if (Asset3DParameters->Asset3DMaterialType != ES3DMasterMaterialType::CUSTOM) {
+		MaterialParams->MasterMaterialName = MasterMaterialType[Asset3DParameters->Asset3DMaterialType];
+		MaterialParams->MasterMaterialPath = GetMasterMaterial(MaterialParams->MasterMaterialName);
+	}
+	else
+	{
+		MaterialParams->MasterMaterialName = TEXT("Custom");
+		MaterialParams->MasterMaterialPath = MatAssetPathsSettings->MasterMaterial3d;
+	}
+	
+	MaterialParams->MaterialInstanceDestination = AssetParams->BaseParams->AssetDestination;
+	MaterialParams->MaterialInstanceName = AssetParams->BaseParams->AssetName + TEXT("_inst");
+	MaterialParams->TexturesDestination = AssetParams->BaseParams->AssetDestination;
+	MaterialParams->bContainsPackedMaps = (AssetImportData->PackedTextures.Num() > 0) ? true : false;
+	Asset3DParameters->MaterialParams.Add(MaterialParams);
+
+	AssetParams->ParamsAssetType = Asset3DParameters;
+	
+	return AssetParams;
+}
+
+ES3DMasterMaterialType FAssetImportParams::GetS3DMasterMaterialType(TSharedPtr<FAssetTypeData> AssetImportData)
+{
+	const UMegascansSettings* MegascansSettings = GetDefault<UMegascansSettings>();	
+	const UMaterialAssetSettings* MatAssetPathsSettings = GetDefault<UMaterialAssetSettings>();
+	
+	if (AssetImportData->AssetMetaInfo->Type == TEXT("3d") && MatAssetPathsSettings->MasterMaterial3d != "None" && MatAssetPathsSettings->MasterMaterial3d != "")
+	{
+		if (UEditorAssetLibrary::DoesAssetExist(MatAssetPathsSettings->MasterMaterial3d))
+		{
+			return ES3DMasterMaterialType::CUSTOM;
+		}
+	}
+	if (AssetImportData->AssetMetaInfo->Type == TEXT("surface") && MatAssetPathsSettings->MasterMaterialSurface != "None" && MatAssetPathsSettings->MasterMaterialSurface != "")
+	{
+		if (UEditorAssetLibrary::DoesAssetExist(MatAssetPathsSettings->MasterMaterialSurface))
+		{
+			return ES3DMasterMaterialType::CUSTOM;
+		}
+
+	}
+	if (AssetImportData->PackedTextures.Num() > 0) {
+		return ES3DMasterMaterialType::CHANNEL_PACK;
+	}
+	if (MegascansSettings->bEnableDisplacement) { 
+		return ES3DMasterMaterialType::NORMAL_DISPLACEMENT; 
+	}
+	// FAssetTextureData type in auto
+	for (auto& TextureComponent : AssetImportData->TextureComponents)
+	{
+		if (TextureComponent->Type == TEXT("fuzz")) {
+			return ES3DMasterMaterialType::NORMAL_FUZZ;
+		}
+	}
+	return ES3DMasterMaterialType::NORMAL;
+}
+
+E3DPlantBillboardMaterial FAssetImportParams::GetBillboardMaterialType(TSharedPtr<FAssetTypeData> AssetImportData)
+{
+	// Type of master material for billboards
+	
+	return (AssetImportData->AssetMetaInfo->bUseBillboardMaterial) ? E3DPlantBillboardMaterial::CAMERA_FACING : E3DPlantBillboardMaterial::NORMAL_BILLBOARD;
+}
+
+
+TSharedPtr<ImportParamsSurfaceAsset> FAssetImportParams::GetSurfaceParams(TSharedPtr<FAssetTypeData> AssetImportData)
+{	
+	TSharedPtr< ImportParamsSurfaceAsset> AssetParams = MakeShareable(new ImportParamsSurfaceAsset);
+	AssetParams->BaseParams = GetAssetImportParams(AssetImportData);
+	const UMaterialAssetSettings* MatAssetPathsSettings = GetDefault<UMaterialAssetSettings>();
+	TSharedPtr<SurfaceParams> SurfaceParameters = MakeShareable( new SurfaceParams);
+	SurfaceParameters->SSubType = GetSurfaceSubtype(AssetImportData);
+	SurfaceParameters->MasterMaterialPath = TEXT("");
+	SurfaceParameters->MasterMaterialName = (SurfaceParameters->SSubType != ESurfaceSubType::SURFACE) ? SurfaceSubTypeMaterials[SurfaceParameters->SSubType] :MasterMaterialType[GetS3DMasterMaterialType(AssetImportData)];
+	if (SurfaceParameters->MasterMaterialName == TEXT("Custom"))
+	{		
+		SurfaceParameters->MasterMaterialPath = MatAssetPathsSettings->MasterMaterialSurface;		
+	}
+	else {
+		SurfaceParameters->MasterMaterialPath = GetMasterMaterial(SurfaceParameters->MasterMaterialName);
+	}		
+	SurfaceParameters->MaterialInstanceDestination = AssetParams->BaseParams->AssetDestination;	
+	SurfaceParameters->MaterialInstanceName = AssetParams->BaseParams->AssetName + TEXT("_inst");	
+	SurfaceParameters->TexturesDestination = AssetParams->BaseParams->AssetDestination;
+	SurfaceParameters->bContainsPackedMaps = (AssetImportData->PackedTextures.Num() > 0) ? true : false;	
+	AssetParams->ParamsAssetType = SurfaceParameters;
+	return AssetParams;
+}
+
+TSharedPtr<ImportParams3DPlantAsset> FAssetImportParams::Get3DPlantParams(TSharedPtr<FAssetTypeData> AssetImportData)
+{	
+	TSharedPtr<ImportParams3DPlantAsset>  AssetParams = MakeShareable(new ImportParams3DPlantAsset);
+	AssetParams->BaseParams = GetAssetImportParams(AssetImportData);
+	AssetParams->ParamsAssetType = MakeShareable(new Asset3DPlantParams);
+
+	const UMaterialAssetSettings* MatAssetPathsSettings = GetDefault<UMaterialAssetSettings>();		
+	AssetParams->ParamsAssetType->FoliageDestination = FPaths::Combine(AssetParams->BaseParams->AssetDestination, TEXT("Foliage"));
+	AssetParams->ParamsAssetType->MeshDestination = AssetParams->BaseParams->AssetDestination;
+	
+	//Plants main master material parameters
+	TSharedPtr < SurfaceParams> MasterMaterialParams = MakeShareable( new SurfaceParams);
+	MasterMaterialParams->MaterialInstanceDestination = AssetParams->BaseParams->AssetDestination;
+	MasterMaterialParams->MaterialInstanceName = AssetParams->BaseParams->AssetName + TEXT("_inst");
+	MasterMaterialParams->TexturesDestination = AssetParams->BaseParams->AssetDestination;
+	MasterMaterialParams->MasterMaterialName = PlantsMasterMaterial;
+	MasterMaterialParams->MasterMaterialPath = GetMasterMaterial(MasterMaterialParams->MasterMaterialName);
+	
+	if (MatAssetPathsSettings->MasterMaterialPlant != "None" && MatAssetPathsSettings->MasterMaterialPlant != "")
+	{
+		if (UEditorAssetLibrary::DoesAssetExist(MatAssetPathsSettings->MasterMaterialPlant))
+		{
+			
+			MasterMaterialParams->MasterMaterialName = TEXT("Custom");
+			MasterMaterialParams->MasterMaterialPath = MatAssetPathsSettings->MasterMaterialPlant;
+		}
+	}
+	//MasterMaterialParams->bContainsPackedMaps = (AssetImportData->PackedTextures.Num() > 0) ? true : false;
+	
+	//MasterMaterialParams->SSubType = ESurfaceSubType::SURFACE;
+	AssetParams->ParamsAssetType->PlantsMasterMaterial = MasterMaterialParams;	
+	
+	//Plants billboard master material parameters
+	TSharedPtr < SurfaceParams> BillboardMasterParams = MakeShareable( new SurfaceParams) ;
+	BillboardMasterParams->MasterMaterialName = BillboardMaterials[GetBillboardMaterialType(AssetImportData)];
+	BillboardMasterParams->MaterialInstanceDestination = FPaths::Combine(AssetParams->BaseParams->AssetDestination, TEXT("Billboard"));	
+	BillboardMasterParams->MaterialInstanceName = AssetParams->BaseParams->AssetName + TEXT("_Billboard_inst");
+	BillboardMasterParams->TexturesDestination = BillboardMasterParams->MaterialInstanceDestination;	
+	BillboardMasterParams->MasterMaterialPath = GetMasterMaterial(BillboardMasterParams->MasterMaterialName);	
+	AssetParams->ParamsAssetType->BillboardMasterMaterial = BillboardMasterParams;	
+	
+	return AssetParams;
+
+}
+
+TSharedPtr<AssetImportParams> FAssetImportParams::GetAssetImportParams(TSharedPtr<FAssetTypeData> AssetImportData)
+{
+	TSharedPtr<AssetImportParams>  AssetImportparameters = MakeShareable(new AssetImportParams);	
+	FString AssetType = AssetImportData->AssetMetaInfo->Type;
+	FString RootDestination = GetRootDestination(AssetImportData->AssetMetaInfo->ExportPath);
+	FString AssetTypePath = GetAssetTypePath(RootDestination, AssetImportData->AssetMetaInfo->Type);
+	AssetImportparameters->AssetName = GetUniqueAssetName(AssetTypePath, RemoveReservedKeywords(NormalizeString(AssetImportData->AssetMetaInfo->Name)));
+	AssetImportparameters->AssetDestination = FPaths::Combine(AssetTypePath, AssetImportparameters->AssetName);
+	return AssetImportparameters;
+}
+
+TSharedPtr<ImportParamsSurfaceAsset> FAssetImportParams::GetAtlasBrushParams(TSharedPtr<FAssetTypeData> AssetImportData)
+{
+	TSharedPtr< ImportParamsSurfaceAsset> AssetParams = MakeShareable(new ImportParamsSurfaceAsset);
+	AssetParams->BaseParams = GetAssetImportParams(AssetImportData);	
+	TSharedPtr<SurfaceParams> SurfaceParameters = MakeShareable(new SurfaceParams);
+	SurfaceParameters->MasterMaterialName = (AssetImportData->AssetMetaInfo->Type == TEXT("atlas")) ? AtlasSubTypeMaterials[GetAtlasSubtype(AssetImportData)] : BrushMasterMaterial;
+	SurfaceParameters->MasterMaterialPath = GetMasterMaterial(SurfaceParameters->MasterMaterialName);	
+	SurfaceParameters->MaterialInstanceDestination = AssetParams->BaseParams->AssetDestination;
+	SurfaceParameters->MaterialInstanceName = AssetParams->BaseParams->AssetName + TEXT("_inst");
+	SurfaceParameters->TexturesDestination = AssetParams->BaseParams->AssetDestination;
+	SurfaceParameters->bContainsPackedMaps = false;
+
+	AssetParams->ParamsAssetType = SurfaceParameters;
+	return AssetParams;
+}
+
+FString FAssetImportParams::GetAssetTypePath(const FString& RootDestination, const FString& AssetType)
+{
+	TMap<FString, FString> TypeNames = {
+		{TEXT("3d"),TEXT("3D_Assets")},
+		{TEXT("3dplant"),TEXT("3D_Plants")},
+		{TEXT("surface"),TEXT("Surfaces")},
+		{TEXT("atlas"),TEXT("Atlases")},
+		{TEXT("brush"),TEXT("Brushes")},
+		{TEXT("3datlas"),TEXT("3D_Atlases")}
+	};
+	return FPaths::Combine(RootDestination, TypeNames[AssetType]);
+}
+
+EAtlasSubType FAssetImportParams::GetAtlasSubtype(TSharedPtr<FAssetTypeData> AssetImportData)
+{
+	auto& TagsList = AssetImportData->AssetMetaInfo->Tags;
+	return (!TagsList.Contains(TEXT("decal"))) ? EAtlasSubType::ATLAS: EAtlasSubType::DECAL;
+}
+ESurfaceSubType FAssetImportParams::GetSurfaceSubtype(TSharedPtr<FAssetTypeData> AssetImportData)
+{
+	auto& TagsList = AssetImportData->AssetMetaInfo->Tags;
+	return (TagsList.Contains(TEXT("metal"))) ? ESurfaceSubType::METAL : (TagsList.Contains(TEXT("surface imperfection"))) ? ESurfaceSubType::IMPERFECTION : (TagsList.Contains(TEXT("tileable displacement"))) ? ESurfaceSubType::DISPLACEMENT : ESurfaceSubType::SURFACE;
+}
+
+EAsset3DSubtype FAssetImportParams::Get3DSubtype(TSharedPtr<FAssetTypeData> AssetImportData)
+{
+	return (AssetImportData->AssetMetaInfo->Categories.Contains(TEXT("scatter")) || AssetImportData->AssetMetaInfo->Tags.Contains(TEXT("scatter")) || AssetImportData->AssetMetaInfo->Categories.Contains(TEXT("cmb_asset")) || AssetImportData->AssetMetaInfo->Tags.Contains(TEXT("cmb_asset"))) ? EAsset3DSubtype::MULTI_MESH : EAsset3DSubtype::NORMAL_3D;
+	
+}
+
+FString FAssetImportParams::GetMasterMaterial(const FString& SelectedMaterial)
+{
+	CopyPresetTextures();
+	FString MaterialPath = GetMaterial(SelectedMaterial);
+	return (MaterialPath == TEXT("")) ? TEXT("") : MaterialPath;
+}
+
+
+
diff --git a/Engine/Plugins/MegascansPlugin/Source/MegascansPlugin/Private/Utilities/AssetData.h b/Engine/Plugins/MegascansPlugin/Source/MegascansPlugin/Private/Utilities/AssetData.h
new file mode 100644
--- /dev/null
+++ b/Engine/Plugins/MegascansPlugin/Source/MegascansPlugin/Private/Utilities/AssetData.h
@@ -0,0 +1,160 @@
+// Copyright Epic Games, Inc. All Rights Reserved.
+#pragma once
+#include "CoreMinimal.h"
+#include "AssetImportData.h"
+
+
+enum EAtlasSubType {
+	DECAL,
+	ATLAS
+};
+
+enum ESurfaceSubType {
+	METAL,
+	IMPERFECTION,
+	DISPLACEMENT,
+	SURFACE
+};
+
+enum EAsset3DSubtype {
+	MULTI_MESH,
+	NORMAL_3D
+};
+
+enum ES3DMasterMaterialType {
+	NORMAL,
+	NORMAL_FUZZ,
+	CHANNEL_PACK,
+	NORMAL_DISPLACEMENT,
+	MODULAR_WINDOWS,
+	CUSTOM
+};
+
+enum E3DPlantBillboardMaterial
+{
+	NORMAL_BILLBOARD,
+	CAMERA_FACING
+};
+
+
+struct SurfaceParams
+{
+	FString MaterialInstanceName;
+	FString MasterMaterialName;
+	FString MasterMaterialPath;
+	FString MaterialInstanceDestination;
+	FString TexturesDestination;
+	ESurfaceSubType SSubType;
+	bool bContainsPackedMaps;
+};
+
+struct Asset3DParams
+{
+	FString MeshDestination;
+	ES3DMasterMaterialType Asset3DMaterialType;
+	EAsset3DSubtype SubType;	
+	TArray<TSharedPtr<SurfaceParams>> MaterialParams;
+};
+
+struct Asset3DPlantParams
+{
+	FString MeshDestination;
+	FString FoliageDestination;
+	E3DPlantBillboardMaterial BillboardMaterialType;
+	TSharedPtr<SurfaceParams> BillboardMasterMaterial;
+	TSharedPtr<SurfaceParams> PlantsMasterMaterial;
+	
+};
+
+struct AssetImportParams
+{	
+	FString AssetName;	
+	FString AssetDestination;	
+};
+
+
+// Main data structures
+struct ImportParams3DAsset
+{
+	TSharedPtr<AssetImportParams> BaseParams;
+	TSharedPtr<Asset3DParams> ParamsAssetType;
+};
+
+struct ImportParams3DPlantAsset
+{
+	TSharedPtr<AssetImportParams> BaseParams;
+	TSharedPtr<Asset3DPlantParams> ParamsAssetType;
+};
+
+struct ImportParamsSurfaceAsset
+{
+	TSharedPtr<AssetImportParams> BaseParams;
+	TSharedPtr<SurfaceParams> ParamsAssetType;
+};
+
+
+
+
+
+class FAssetImportParams {
+private:
+	FAssetImportParams() = default;
+
+	const TMap<EAtlasSubType, FString> AtlasSubTypeMaterials = {
+		{EAtlasSubType::ATLAS,"Atlas_Material"},
+		{EAtlasSubType::DECAL,"Decal_MasterMaterial"}
+	};
+
+	const TMap<E3DPlantBillboardMaterial, FString> BillboardMaterials = {
+		{E3DPlantBillboardMaterial::CAMERA_FACING,"Billboard_Material"},
+		{E3DPlantBillboardMaterial::NORMAL_BILLBOARD,"Foliage_Material"}
+	};
+	//These values should be stored in a configuration file
+	const TMap<ESurfaceSubType, FString> SurfaceSubTypeMaterials = { 
+		{ESurfaceSubType::METAL,"MS_DefaultMaterial"},
+		{ESurfaceSubType::IMPERFECTION,"Imperfection_MasterMaterial"} ,
+		{ESurfaceSubType::DISPLACEMENT,"Displacement_MasterMaterial"}
+	};
+
+	const FString BrushMasterMaterial = TEXT("BrushDecal_Material");
+
+	const TMap<ES3DMasterMaterialType, FString> MasterMaterialType = {
+		{ES3DMasterMaterialType::CHANNEL_PACK , TEXT("MS_DefaultMaterial_CP")},
+		{ES3DMasterMaterialType::MODULAR_WINDOWS , TEXT("")},
+		{ES3DMasterMaterialType::NORMAL , TEXT("MS_DefaultMaterial")},
+		{ES3DMasterMaterialType::NORMAL_DISPLACEMENT , TEXT("MS_DefaultMaterial_Displacement")},
+		{ES3DMasterMaterialType::NORMAL_FUZZ , TEXT("MS_DefaultMaterial_Fuzzy")},
+		{ES3DMasterMaterialType::CUSTOM , TEXT("Custom")}
+	};
+
+	const FString PlantsMasterMaterial = TEXT("Foliage_Material");
+	
+	static TSharedPtr<FAssetImportParams> ImportParamsInst;
+	FString GetAssetTypePath(const FString& RootDestination, const FString& AssetType);
+	EAtlasSubType GetAtlasSubtype(TSharedPtr<FAssetTypeData> AssetImportData);
+	ESurfaceSubType GetSurfaceSubtype(TSharedPtr<FAssetTypeData> AssetImportData);
+	EAsset3DSubtype Get3DSubtype(TSharedPtr<FAssetTypeData> AssetImportData);
+	ES3DMasterMaterialType GetS3DMasterMaterialType(TSharedPtr<FAssetTypeData> AssetImportData);
+	E3DPlantBillboardMaterial GetBillboardMaterialType(TSharedPtr<FAssetTypeData> AssetImportData);
+	FString GetMasterMaterial(const FString& SelectedMaterial);
+
+public:
+	static TSharedPtr<FAssetImportParams> Get();	
+//Asset related data
+	TSharedPtr<ImportParams3DAsset> Get3DAssetsParams(TSharedPtr<FAssetTypeData> AssetImportData);
+	TSharedPtr<ImportParamsSurfaceAsset> GetSurfaceParams(TSharedPtr<FAssetTypeData> AssetImportData);
+	TSharedPtr<ImportParams3DPlantAsset> Get3DPlantParams(TSharedPtr<FAssetTypeData> AssetImportData);
+	TSharedPtr<AssetImportParams> GetAssetImportParams(TSharedPtr<FAssetTypeData> AssetImportData);
+	TSharedPtr<ImportParamsSurfaceAsset> GetAtlasBrushParams(TSharedPtr<FAssetTypeData> AssetImportData);
+};
+
+
+
+
+
+
+
+
+
+
+
diff --git a/Engine/Plugins/MegascansPlugin/Source/MegascansPlugin/Private/Utilities/AssetsDatabase.cpp b/Engine/Plugins/MegascansPlugin/Source/MegascansPlugin/Private/Utilities/AssetsDatabase.cpp
new file mode 100644
--- /dev/null
+++ b/Engine/Plugins/MegascansPlugin/Source/MegascansPlugin/Private/Utilities/AssetsDatabase.cpp
@@ -0,0 +1,73 @@
+// Copyright Epic Games, Inc. All Rights Reserved.
+#include "AssetsDatabase.h"
+#include "Misc/Paths.h"
+//#include "SQLiteDatabase.h"
+#include "SQLitePreparedStatement.h"
+#include "SQLiteDatabaseConnection.h"
+#include "SQLiteResultSet.h"
+
+TSharedPtr<FAssetsDatabase> FAssetsDatabase::DBInst;
+
+TSharedPtr<FAssetsDatabase> FAssetsDatabase::Get()
+{
+	if (!DBInst.IsValid())
+	{
+		DBInst = MakeShareable(new FAssetsDatabase);
+	}
+	return DBInst;
+}
+
+FAssetsDatabase::FAssetsDatabase()
+{
+	SQLiteDatabase = MakeShareable(new FSQLiteDatabaseConnection);	
+	FString DatabasePath = FPaths::Combine(FPaths::ConvertRelativePathToFull(FPaths::ProjectContentDir()), TEXT("MSPresets"), TEXT("MSAssets.db"));	
+	CreateAssetsDatabase(DatabasePath);	
+}
+
+FAssetsDatabase::~FAssetsDatabase()
+{
+	SQLiteDatabase->Close();
+}
+
+void FAssetsDatabase::CreateAssetsDatabase(const FString& DatabasePath)
+{		
+	SQLiteDatabase->Open(*DatabasePath, nullptr, nullptr);
+	FString CreateStatement = TEXT("CREATE TABLE IF NOT EXISTS MegascansAssets("
+		"ID TEXT PRIMARY KEY     NOT NULL,"
+		"NAME           TEXT    NOT NULL,"
+		"PATH            TEXT     NOT NULL,"
+		"TYPE            TEXT     NOT NULL"
+		");");
+
+	SQLiteDatabase->Execute(*CreateStatement);
+}
+
+void FAssetsDatabase::AddRecord(const AssetRecord& Record)
+{
+	FString InsertStatement = FString::Printf(TEXT("INSERT OR REPLACE INTO MegascansAssets VALUES('%s', '%s', '%s', '%s');"), *Record.ID, *Record.Name, *Record.Path, *Record.Type);
+	
+	
+	SQLiteDatabase->Execute(*InsertStatement);
+}
+
+bool FAssetsDatabase::RecordExists(const FString& AssetID, AssetRecord& Record)
+{
+
+	FString SelectStatement = FString::Printf(TEXT("SELECT * FROM MegascansAssets WHERE ID = '%s'"), *AssetID);
+	
+	FSQLiteResultSet* QueryResult = NULL;
+
+	if (SQLiteDatabase->Execute(*SelectStatement, QueryResult))
+	{
+		for (FSQLiteResultSet::TIterator ResultIterator(QueryResult); ResultIterator; ++ResultIterator)
+		{
+			Record.ID = ResultIterator->GetString(TEXT("ID"));
+			Record.Name = ResultIterator->GetString(TEXT("NAME"));
+			Record.Path = ResultIterator->GetString(TEXT("PATH"));
+			Record.Type = ResultIterator->GetString(TEXT("TYPE"));
+			return true;
+		}
+
+	}	
+	return false;
+}
diff --git a/Engine/Plugins/MegascansPlugin/Source/MegascansPlugin/Private/Utilities/AssetsDatabase.h b/Engine/Plugins/MegascansPlugin/Source/MegascansPlugin/Private/Utilities/AssetsDatabase.h
new file mode 100644
--- /dev/null
+++ b/Engine/Plugins/MegascansPlugin/Source/MegascansPlugin/Private/Utilities/AssetsDatabase.h
@@ -0,0 +1,32 @@
+// Copyright Epic Games, Inc. All Rights Reserved.
+#pragma once
+#include "CoreMinimal.h"
+
+
+class FSQLiteDatabase;
+class FSQLiteDatabaseConnection;
+
+struct AssetRecord {
+	FString ID;
+	FString Name;
+	FString Path;
+	FString Type;
+};
+
+class FAssetsDatabase
+{
+private:
+	FAssetsDatabase() ;
+	
+	static TSharedPtr<FAssetsDatabase> DBInst;
+	void CreateAssetsDatabase(const FString& DatabasePath);
+	TSharedPtr<FSQLiteDatabaseConnection> SQLiteDatabase;
+
+
+public:
+	~FAssetsDatabase();
+	static TSharedPtr<FAssetsDatabase> Get();
+	void AddRecord(const AssetRecord& Record);
+	bool RecordExists(const FString& AssetID, AssetRecord& Record);
+
+};
\ No newline at end of file
diff --git a/Engine/Plugins/MegascansPlugin/Source/MegascansPlugin/Private/Utilities/MaterialUtils.cpp b/Engine/Plugins/MegascansPlugin/Source/MegascansPlugin/Private/Utilities/MaterialUtils.cpp
new file mode 100644
--- /dev/null
+++ b/Engine/Plugins/MegascansPlugin/Source/MegascansPlugin/Private/Utilities/MaterialUtils.cpp
@@ -0,0 +1,33 @@
+// Copyright Epic Games, Inc. All Rights Reserved.
+#include "Utilities/MaterialUtils.h"
+#include "Editor/UnrealEd/Classes/Factories/MaterialInstanceConstantFactoryNew.h"
+
+#include "Runtime/Engine/Classes/Engine/Texture.h"
+#include "MaterialEditingLibrary.h"
+#include "AssetToolsModule.h"
+#include "IAssetTools.h"
+#include "EditorAssetLibrary.h"
+
+
+#include "UI/MSSettings.h"
+
+
+UMaterialInstanceConstant* FMaterialUtils::CreateInstanceMaterial(const FString& MasterMaterialPath, const FString& InstanceDestination, const FString& AssetName)
+{	
+	UMaterialInterface* MasterMaterial = CastChecked<UMaterialInterface>(UEditorAssetLibrary::LoadAsset(MasterMaterialPath));
+	if (MasterMaterial == nullptr) return nullptr;
+
+	IAssetTools& AssetTools = FModuleManager::GetModuleChecked<FAssetToolsModule>("AssetTools").Get();
+	UMaterialInstanceConstantFactoryNew* Factory = NewObject<UMaterialInstanceConstantFactoryNew>();
+	UMaterialInstanceConstant* MaterialInstance = CastChecked<UMaterialInstanceConstant>(AssetTools.CreateAsset(AssetName, InstanceDestination, UMaterialInstanceConstant::StaticClass(), Factory));
+	UMaterialEditingLibrary::SetMaterialInstanceParent(MaterialInstance, MasterMaterial);
+	if (MaterialInstance)
+	{
+		MaterialInstance->SetFlags(RF_Standalone);
+		MaterialInstance->MarkPackageDirty();
+		MaterialInstance->PostEditChange();
+	}
+	return MaterialInstance;
+}
+
+
diff --git a/Engine/Plugins/MegascansPlugin/Source/MegascansPlugin/Private/Utilities/MaterialUtils.h b/Engine/Plugins/MegascansPlugin/Source/MegascansPlugin/Private/Utilities/MaterialUtils.h
new file mode 100644
--- /dev/null
+++ b/Engine/Plugins/MegascansPlugin/Source/MegascansPlugin/Private/Utilities/MaterialUtils.h
@@ -0,0 +1,14 @@
+// Copyright Epic Games, Inc. All Rights Reserved.
+#pragma once
+#include "CoreMinimal.h"
+
+#include "AssetImportData.h"
+#include "Materials/MaterialInstanceConstant.h"
+
+class FMaterialUtils
+{
+public:
+	static UMaterialInstanceConstant* CreateInstanceMaterial(const FString& MasterMaterialPath, const FString& InstanceDestination, const FString& AssetName);
+	
+
+};
\ No newline at end of file
diff --git a/Engine/Plugins/MegascansPlugin/Source/MegascansPlugin/Private/Utilities/MeshOp.cpp b/Engine/Plugins/MegascansPlugin/Source/MegascansPlugin/Private/Utilities/MeshOp.cpp
new file mode 100644
--- /dev/null
+++ b/Engine/Plugins/MegascansPlugin/Source/MegascansPlugin/Private/Utilities/MeshOp.cpp
@@ -0,0 +1,240 @@
+// Copyright Epic Games, Inc. All Rights Reserved.
+#include "Utilities/MeshOp.h"
+#include "UnrealEd/Classes/Factories/FbxImportUI.h"
+#include "UnrealEd/Classes/Factories/FbxStaticMeshImportData.h"
+#include "AssetImportTask.h"
+#include "AssetToolsModule.h"
+#include "IAssetTools.h"
+#include "FbxMeshUtils.h"
+#include "AbcImportSettings.h"
+#include "Utilities/MiscUtils.h"
+#include "Runtime/Core/Public/Misc/Paths.h"
+
+#include "EditorAssetLibrary.h"
+#include "EditorStaticMeshLibrary.h"
+#include "Editor/FoliageEdit/Private/FoliageTypeFactory.h"
+#include "Runtime/Foliage/Public/InstancedFoliageActor.h"
+
+#include "Runtime/Foliage/Public/FoliageType.h"
+#include "Runtime/Foliage/Public/FoliageType_InstancedStaticMesh.h"
+#include "Runtime/AssetRegistry/Public/AssetRegistryModule.h"
+#include "PerPlatformProperties.h"
+
+#include "StaticMeshEditorSubsystem.h"
+
+
+TSharedPtr<FMeshOps> FMeshOps::MeshOpsInst;
+
+
+
+TSharedPtr<FMeshOps> FMeshOps::Get()
+{
+	if (!MeshOpsInst.IsValid())
+	{
+		MeshOpsInst = MakeShareable(new FMeshOps);
+	}
+	return MeshOpsInst;
+}
+
+FString FMeshOps::ImportMesh(const FString& MeshPath, const FString& Destination, const FString& AssetName)
+{
+	
+	FString FileExtension;
+	FString FilePath;
+	FString AssetPath;
+
+	UFbxImportUI* ImportOptions;
+	MeshPath.Split(TEXT("."), &FilePath, &FileExtension);
+	
+
+	FileExtension = FPaths::GetExtension(MeshPath);
+	if (FileExtension == TEXT("obj") || FileExtension == TEXT("fbx"))
+	{
+		ImportOptions = GetFbxOptions();
+		if (FileExtension == TEXT("obj"))
+		{
+			ImportOptions->StaticMeshImportData->ImportRotation = FRotator(0, 0, 90);
+		}
+		
+
+		AssetPath = ImportFile(ImportOptions, Destination, AssetName, MeshPath);
+	}
+	else if (FileExtension == TEXT("abc"))
+	{
+		UAbcImportSettings* AbcImportOptions = GetAbcSettings();
+		AssetPath = ImportFile(AbcImportOptions, Destination, AssetName, MeshPath);
+
+	}
+	return AssetPath;
+}
+
+void FMeshOps::ApplyLods(const TArray<FString>& LodList, UStaticMesh* SourceMesh)
+{
+	int32 LodCounter = 1;
+	for (FString LodPath : LodList)
+	{
+		if (LodCounter > 7) continue;
+		FbxMeshUtils::ImportStaticMeshLOD(SourceMesh, LodPath, LodCounter);
+		
+		LodCounter++;
+	}
+
+	
+}
+
+void FMeshOps::ApplyAbcLods(UStaticMesh* SourceMesh, const TArray<FString>& LodPathList, const FString& AssetDestination)
+{
+	FString LodDestination = FPaths::Combine(AssetDestination, TEXT("Lods"));
+	int32 LodCounter = 1;
+	UStaticMeshEditorSubsystem* SMESubsystem = GEditor->GetEditorSubsystem<UStaticMeshEditorSubsystem>();
+
+	for (FString LodPath : LodPathList)
+	{
+		if (LodCounter > 7) continue;
+		FString ImportedLodPath = ImportMesh(LodPath, LodDestination, "");
+		UStaticMesh* ImportedLod = CastChecked<UStaticMesh>(LoadAsset(ImportedLodPath));
+		//UEditorStaticMeshLibrary::SetLodFromStaticMesh(SourceMesh, LodCounter, ImportedLod, 0, true);
+		SMESubsystem->SetLodFromStaticMesh(SourceMesh, LodCounter, ImportedLod, 0, true);
+
+		LodCounter++;
+	}
+	UEditorAssetLibrary::DeleteDirectory(LodDestination);
+}
+
+void FMeshOps::CreateFoliageAsset(const FString& FoliagePath, UStaticMesh* SourceAsset, const FString& FoliageAssetName, bool bSavePackage)
+{
+	
+	if (SourceAsset == nullptr) return;
+	FString FoliageTypePath = FPaths::Combine(FoliagePath, TEXT("Foliage/"));
+	//IAssetTools& AssetTools = FModuleManager::GetModuleChecked<FAssetToolsModule>("AssetTools").Get();
+	FString PackageName = FoliageTypePath;
+	PackageName += FoliageAssetName;
+	UPackage* Package = CreatePackage(NULL, *PackageName);
+	UFoliageType_InstancedStaticMesh* FoliageAsset = NewObject<UFoliageType_InstancedStaticMesh>(Package, *FoliageAssetName, RF_Public | RF_Standalone | RF_MarkAsRootSet);
+	FoliageAsset->SetStaticMesh(SourceAsset);
+	FAssetRegistryModule::AssetCreated(FoliageAsset);
+	FoliageAsset->PostEditChange();
+	Package->MarkPackageDirty();
+	//AssetUtils::SavePackage(FoliageAsset);
+	auto* CurrentWorld = GEditor->GetEditorWorldContext().World();
+	AInstancedFoliageActor* IFA = AInstancedFoliageActor::GetInstancedFoliageActorForCurrentLevel(CurrentWorld, true);
+	IFA->AddFoliageType(FoliageAsset);
+	
+}
+
+
+
+TArray<FString> FMeshOps::ImportLodsAsStaticMesh(const TArray<FString> & LodList, const FString & AssetDestination)
+{
+
+	return TArray<FString>();
+}
+
+
+UFbxImportUI* FMeshOps::GetFbxOptions()
+{
+	UFbxImportUI* ImportOptions = NewObject<UFbxImportUI>();
+	ImportOptions->bImportMesh = true;
+	ImportOptions->bImportAnimations = false;
+	ImportOptions->bImportMaterials = false;
+	ImportOptions->bImportAsSkeletal = false;
+	ImportOptions->StaticMeshImportData->bCombineMeshes = false;
+	ImportOptions->StaticMeshImportData->bGenerateLightmapUVs = false;
+	ImportOptions->StaticMeshImportData->bAutoGenerateCollision = false;
+	ImportOptions->StaticMeshImportData->VertexColorImportOption = EVertexColorImportOption::Replace;
+	
+
+	return ImportOptions;
+}
+
+void FMeshOps::RemoveExtraMaterialSlot(UStaticMesh* SourceMesh)
+{
+
+	for (int i = 1; i < SourceMesh->GetNumLODs(); i++)
+	{
+		
+		FMeshSectionInfo MeshSectionInfo = SourceMesh->GetSectionInfoMap().Get(i, 0);
+		MeshSectionInfo.MaterialIndex = 0;
+		
+		SourceMesh->GetSectionInfoMap().Set(i, 0, MeshSectionInfo);
+		
+	}
+
+	int NumOfMaterials = SourceMesh->StaticMaterials.Num();
+	for (int i = NumOfMaterials; i > 1; i--)
+	{
+		SourceMesh->StaticMaterials.RemoveAt(i - 1);
+	}
+
+	SourceMesh->Modify();
+	SourceMesh->PostEditChange();
+	SourceMesh->MarkPackageDirty();
+
+}
+
+
+template<class T>
+FString FMeshOps::ImportFile(T * ImportOptions, const FString & Destination, const FString & AssetName, const FString & Source)
+{
+
+	
+
+	FString ImportedMeshPath = TEXT("");
+	TArray< UAssetImportTask*> ImportTasks;
+	UAssetImportTask* MeshImportTask = NewObject<UAssetImportTask>();
+	MeshImportTask->bAutomated = true;
+	MeshImportTask->bSave = false;
+	MeshImportTask->Filename = Source;
+
+	
+	
+
+	MeshImportTask->DestinationName = AssetName;
+	MeshImportTask->DestinationPath = Destination;
+	MeshImportTask->bReplaceExisting = true;
+	MeshImportTask->Options = ImportOptions;
+
+	ImportTasks.Add(MeshImportTask);
+
+	IAssetTools& AssetTools = FModuleManager::GetModuleChecked<FAssetToolsModule>("AssetTools").Get();
+
+	AssetTools.ImportAssetTasks(ImportTasks);
+
+	for (UAssetImportTask* ImpTask : ImportTasks)
+	{
+		if(ImpTask->ImportedObjectPaths.Num()>0)
+			ImportedMeshPath = ImpTask->ImportedObjectPaths[0];
+	}
+	return ImportedMeshPath;
+}
+
+UAbcImportSettings* FMeshOps::GetAbcSettings()
+{
+	UAbcImportSettings* AlembicOptions = NewObject<UAbcImportSettings>();
+	AlembicOptions->ImportType = EAlembicImportType::StaticMesh;
+	AlembicOptions->StaticMeshSettings.bMergeMeshes = false;
+	AlembicOptions->MaterialSettings.bCreateMaterials = false;
+	AlembicOptions->SamplingSettings.FrameStart = 0;
+	AlembicOptions->SamplingSettings.FrameEnd = 1;
+	AlembicOptions->ConversionSettings.Rotation = FVector(90, 00, 00);
+	AlembicOptions->ConversionSettings.Scale = FVector(1.0, -1.0, 1.0);
+	AlembicOptions->ConversionSettings.bFlipU = false;
+	AlembicOptions->ConversionSettings.bFlipV = true;
+	return AlembicOptions;
+}
+
+void FMeshOps::LodDistanceTest(UStaticMesh* SourceMesh)
+{
+	SourceMesh->bAutoComputeLODScreenSize = false;
+	FStaticMeshSourceModel& SourceModel = SourceMesh->GetSourceModel(0);
+	SourceModel.ScreenSize = FPerPlatformFloat(1);
+	FStaticMeshSourceModel& SourceModel1 = SourceMesh->GetSourceModel(1);
+	SourceModel1.ScreenSize = FPerPlatformFloat(0.5);
+
+
+}
+
+
+
+
+
diff --git a/Engine/Plugins/MegascansPlugin/Source/MegascansPlugin/Private/Utilities/MeshOp.h b/Engine/Plugins/MegascansPlugin/Source/MegascansPlugin/Private/Utilities/MeshOp.h
new file mode 100644
--- /dev/null
+++ b/Engine/Plugins/MegascansPlugin/Source/MegascansPlugin/Private/Utilities/MeshOp.h
@@ -0,0 +1,32 @@
+// Copyright Epic Games, Inc. All Rights Reserved.
+#pragma once
+#include "CoreMinimal.h"
+
+
+class UFbxImportUI;
+class UStaticMesh;
+class UAbcImportSettings;
+
+
+class FMeshOps
+{
+private:
+	FMeshOps() = default;
+	static TSharedPtr<FMeshOps> MeshOpsInst;
+	UFbxImportUI* GetFbxOptions();
+	
+
+	template<class T> FString ImportFile(T* ImportOptions, const FString& Destination, const FString& AssetName, const FString& Source);
+	UAbcImportSettings* GetAbcSettings();
+
+public:
+	static TSharedPtr<FMeshOps> Get();
+	FString ImportMesh(const FString& MeshPath, const FString& Destination, const FString& AssetName = TEXT(""));
+	void ApplyLods(const TArray<FString>& LodList, UStaticMesh* SourceMesh);
+	TArray<FString> ImportLodsAsStaticMesh(const TArray<FString>& LodList, const FString& AssetDestination);
+	void ApplyAbcLods(UStaticMesh* SourceMesh, const TArray<FString>& LodPathList, const FString& AssetDestination);
+	void CreateFoliageAsset(const FString& FoliagePath, UStaticMesh* SourceAsset, const FString& FoliageAssetName, bool bSavePackage = false);
+	void RemoveExtraMaterialSlot(UStaticMesh* SourceMesh);
+
+	void LodDistanceTest(UStaticMesh* SourceMesh);
+};
\ No newline at end of file
diff --git a/Engine/Plugins/MegascansPlugin/Source/MegascansPlugin/Private/Utilities/MiscUtils.cpp b/Engine/Plugins/MegascansPlugin/Source/MegascansPlugin/Private/Utilities/MiscUtils.cpp
new file mode 100644
--- /dev/null
+++ b/Engine/Plugins/MegascansPlugin/Source/MegascansPlugin/Private/Utilities/MiscUtils.cpp
@@ -0,0 +1,606 @@
+// Copyright Epic Games, Inc. All Rights Reserved.
+#include "Utilities/MiscUtils.h"
+
+#include "Runtime/Core/Public/Misc/Paths.h"
+#include "Runtime/Core/Public/HAL/PlatformFilemanager.h"
+#include "Runtime/Core/Public/GenericPlatform/GenericPlatformFile.h"
+#include "Runtime/Core/Public/HAL/FileManager.h"
+
+#include "EditorAssetLibrary.h"
+#include "Runtime/AssetRegistry/Public/AssetRegistryModule.h"
+
+#include "ContentBrowserModule.h"
+#include "IContentBrowserSingleton.h"
+#include "Runtime/Engine/Classes/Engine/Selection.h"
+
+#include "Editor/UnrealEd/Classes/Factories/MaterialInstanceConstantFactoryNew.h"
+
+#include "Runtime/Engine/Classes/Engine/StaticMesh.h"
+
+#include "Runtime/Core/Public/Misc/MessageDialog.h"
+#include "Runtime/Core/Public/Internationalization/Text.h"
+
+#include "PackageTools.h"
+
+#include <regex>
+#include <ostream>
+#include <sstream>
+
+#include "Misc/ScopedSlowTask.h"
+
+
+
+TSharedPtr<FJsonObject> DeserializeJson(const FString& JsonStringData)
+{
+	TSharedPtr<FJsonObject> JsonDataObject;
+	bool bIsDeseriliazeSuccessful = false;	
+	TSharedRef<TJsonReader<>> JsonReader = TJsonReaderFactory<>::Create(JsonStringData);
+	bIsDeseriliazeSuccessful = FJsonSerializer::Deserialize(JsonReader, JsonDataObject);	
+	return JsonDataObject;
+}
+
+void ShowErrorDialog(const FString& ErrorMessage)
+{
+
+}
+
+FString GetPluginPath()
+{
+	FString PluginName = TEXT("MegascansPlugin");
+	FString PluginsPath = FPaths::EnginePluginsDir();
+
+	FString MSPluginPath = FPaths::Combine(PluginsPath, PluginName);
+	return MSPluginPath;
+}
+
+FString GetSourceMSPresetsPath()
+{
+	FString MaterialPresetsPath = FPaths::Combine(GetPluginPath(), TEXT("Content"), GetMSPresetsName());
+	return MaterialPresetsPath;
+}
+
+FString GetPreferencesPath()
+{
+	FString PrefFileName = TEXT("unreal_action_settings.json");
+	return FPaths::Combine(GetPluginPath(), TEXT("Content"), TEXT("Python"), PrefFileName);
+}
+
+FString GetMSPresetsName()
+{
+	return TEXT("MSPresets");
+}
+
+FString GetMaterial(const FString & MaterialName)
+{	
+
+	FString MaterialPath = FPaths::Combine(TEXT("/Game"), GetMSPresetsName(), MaterialName, MaterialName);
+	
+	if (!UEditorAssetLibrary::DoesAssetExist(MaterialPath))
+	{
+		//UE_LOG(LogTemp, Error, TEXT("Master Material doesnt exist."));
+		if (CopyMaterialPreset(MaterialName))
+		{
+			if (UEditorAssetLibrary::DoesAssetExist(MaterialPath))
+			{	
+				
+				
+				return MaterialPath;
+			}			
+			
+		}
+	}
+	else {
+		return MaterialPath;
+	}
+	return TEXT("");
+}
+
+bool CopyMaterialPreset(const FString & MaterialName)
+{	
+	FString MaterialSourceFolderPath = FPaths::Combine(GetSourceMSPresetsPath(), MaterialName);
+	MaterialSourceFolderPath = FPaths::ConvertRelativePathToFull(MaterialSourceFolderPath);	
+	IPlatformFile& PlatformFile = FPlatformFileManager::Get().GetPlatformFile();	
+	if (!PlatformFile.DirectoryExists(*MaterialSourceFolderPath))
+	{
+		return false;
+	}
+
+	FString MaterialDestinationPath = FPaths::ProjectContentDir();
+	MaterialDestinationPath = FPaths::Combine(MaterialDestinationPath, GetMSPresetsName());
+	MaterialDestinationPath = FPaths::Combine(MaterialDestinationPath, MaterialName);
+	
+	if (!PlatformFile.DirectoryExists(*MaterialDestinationPath))
+	{
+		
+		if (!PlatformFile.CreateDirectoryTree(*MaterialDestinationPath))
+		{
+			
+			return false;
+		}
+
+		if (!PlatformFile.CopyDirectoryTree(*MaterialDestinationPath, *MaterialSourceFolderPath, true))
+		{
+			
+			return false;
+		}
+
+	}
+
+	
+	FAssetRegistryModule& AssetRegistryModule = FModuleManager::GetModuleChecked<FAssetRegistryModule>("AssetRegistry");
+	TArray<FString> MaterialBasePath;
+	FString BasePath = FPaths::Combine(FPaths::ProjectContentDir(), GetMSPresetsName());
+	MaterialBasePath.Add("/Game/MSPresets");
+	MaterialBasePath.Add(FPaths::Combine(TEXT("/Game/MSPresets"), MaterialName));
+	AssetRegistryModule.Get().ScanPathsSynchronous(MaterialBasePath, true);
+	return true;
+	
+}
+
+
+
+bool CopyPresetTextures()
+{
+	FString TexturesDestinationPath = FPaths::ProjectContentDir();
+	
+	TexturesDestinationPath = FPaths::Combine(TexturesDestinationPath, GetMSPresetsName(), TEXT("MSTextures"));
+
+	FString TexturesSourceFolderPath = FPaths::Combine(GetSourceMSPresetsPath(), TEXT("MSTextures"));	
+	IPlatformFile& PlatformFile = FPlatformFileManager::Get().GetPlatformFile();
+	if (!PlatformFile.DirectoryExists(*TexturesDestinationPath))
+	{
+		if (!PlatformFile.CreateDirectoryTree(*TexturesDestinationPath))
+		{
+			
+			return false;
+		}
+
+		if (!PlatformFile.CopyDirectoryTree(*TexturesDestinationPath, *TexturesSourceFolderPath, true))
+		{
+			
+			return false;
+		}
+	}
+
+	return true;
+}
+
+
+UObject* LoadAsset(const FString& AssetPath)
+{
+	return UEditorAssetLibrary::LoadAsset(AssetPath);
+}
+
+TArray<FString> GetAssetsList(const FString& DirectoryPath)
+{
+	return UEditorAssetLibrary::ListAssets(DirectoryPath, false);
+
+}
+void SaveAsset(const FString& AssetPath) 
+{
+	UEditorAssetLibrary::SaveAsset(AssetPath);
+}
+
+
+void DeleteExtraMesh(const FString& BasePath)
+{
+	TArray<FString> AssetBasePath;
+	AssetBasePath.Add(BasePath);
+	FAssetRegistryModule& AssetRegistryModule = FModuleManager::GetModuleChecked<FAssetRegistryModule>("AssetRegistry");
+	AssetRegistryModule.Get().ScanPathsSynchronous(AssetBasePath, true);
+	FString ExtraMeshName = TEXT("StaticMesh_0");
+	FString AssetPath = FPaths::Combine(BasePath, ExtraMeshName);
+	
+	AssetRegistryModule.Get().ScanPathsSynchronous(AssetBasePath, true);
+}
+
+
+
+
+
+FString ResolvePath(const FString& PathConvention, TSharedPtr<FAssetTypeData> AssetsImportData, TSharedPtr<FAssetTextureData> TextureData)
+{
+	TMap<FString, FString> TypeNames = {
+		{TEXT("3d"),TEXT("3D_Assets")},
+		{TEXT("3dplant"),TEXT("3D_Plants")},
+		{TEXT("surface"),TEXT("Surfaces")},
+		{TEXT("atlas"),TEXT("Atlases")},
+		{TEXT("brush"),TEXT("Brushes")},
+		{TEXT("3datlas"),TEXT("3D_Atlases")}
+	};
+
+	
+	FString ResolvedPath = TEXT("");
+	FString ResolvedConvention = TEXT("");
+	
+	FString PlaceholderConventionTexture = "Megascans/{AssetType}/{AssetName}/";
+
+	TArray<FString> NamingConventions;
+	PathConvention.ParseIntoArray(NamingConventions, TEXT("/"));
+
+	for (FString NamingConvention : NamingConventions)
+	{
+		if (NamingConvention.Contains(TEXT("{")) || NamingConvention.Contains(TEXT("}")))
+		{
+			
+			ResolvedConvention = TEXT("");
+			NamingConvention = NamingConvention.Replace(TEXT("{"), TEXT(""));
+			NamingConvention = NamingConvention.Replace(TEXT("}"), TEXT(""));
+
+			if (NamingConvention == TEXT("AssetName"))
+			{
+				ResolvedConvention = AssetsImportData->AssetMetaInfo->Name;
+			}
+			else if (NamingConvention == TEXT("AssetType"))
+			{
+				if (TypeNames.Contains(AssetsImportData->AssetMetaInfo->Type))
+				{
+					ResolvedConvention = TypeNames[AssetsImportData->AssetMetaInfo->Type];
+				}
+				else {
+					ResolvedConvention = TEXT("Misc");
+				}
+			}
+			else if (NamingConvention == TEXT("TextureType"))
+			{
+				if (TextureData)
+				{
+					ResolvedConvention = TextureData->Type;
+				}
+				else {
+					ResolvedConvention = TEXT("TextureType");
+				}
+			}
+			else if (NamingConvention == TEXT("AssetId"))
+			{
+				ResolvedConvention = AssetsImportData->AssetMetaInfo->Id;
+			}
+			else if (NamingConvention == TEXT("TextureResolution"))
+			{
+				if (TextureData)
+				{
+					ResolvedConvention = TextureData->Resolution;
+				}
+				else {
+					ResolvedConvention = TEXT("Resolution");
+				}
+			}
+			else if (NamingConvention == TEXT("TextureName"))
+			{
+				if (TextureData)
+				{
+					ResolvedConvention = TextureData->Name;
+				}
+				else {
+					ResolvedConvention = TEXT("Name");
+				}
+			}
+
+			ResolvedConvention = NormalizeString(ResolvedConvention);
+			ResolvedPath = FPaths::Combine(ResolvedPath, ResolvedConvention);
+		}
+		else 
+		{
+			ResolvedPath = FPaths::Combine(ResolvedPath, NamingConvention);
+		}
+
+	}
+
+	return ResolvedPath;
+}
+
+FString ResolveName(const FString& NamingConvention, TSharedPtr<FAssetTypeData> AssetsImportData, TSharedPtr<FAssetTextureData> TextureData)
+{
+	
+	FString ResolvedName = TEXT("");
+	FString DummyNameConvention = TEXT("prefix_{AssetName}{AssetType}_{TextureType}_{TextureName}_postfix");
+	return ResolvedName;
+}
+
+
+FString NormalizeString(FString InputString)
+{	
+	const std::string SInputString = std::string(TCHAR_TO_UTF8(*InputString));	
+	const std::regex SpecialCharacters("[^a-zA-Z0-9]+");
+	std::stringstream Result;
+	std::regex_replace(std::ostream_iterator<char>(Result), SInputString.begin(), SInputString.end(), SpecialCharacters, "_");
+	FString NormalizedString = FString(Result.str().c_str());
+
+	return NormalizedString;
+}
+
+FString RemoveReservedKeywords(const FString& Name)
+{
+	FString ResolvedName = Name;
+	TArray<FString> ReservedKeywrods = { "CON","PRN", "AUX", "CLOCK$", "NUL", "NONE", "COM1", "COM2", "COM3", "COM4", "COM5", "COM6", "COM7", "COM8", "COM9", "LPT1", "LPT2", "LPT3", "LPT4", "LPT5", "LPT6", "LPT7", "LPT8", "LPT9" };
+	TArray<FString> NameArray;
+	Name.ParseIntoArray(NameArray, TEXT("_"));
+	for (FString NameTag : NameArray)
+	{
+		if (ReservedKeywrods.Contains(NameTag))
+		{			
+			ResolvedName.Replace(*NameTag, TEXT(""));
+		}
+	}
+	return ResolvedName;
+}
+
+
+TArray<FString> ParseLodList(TSharedPtr<FAssetTypeData> AssetImportData)
+{
+	
+	FString FileExtension;
+	FString FilePath;
+	AssetImportData->MeshList[0]->Path.Split(TEXT("."), &FilePath, &FileExtension);
+	FileExtension = FPaths::GetExtension(AssetImportData->MeshList[0]->Path);
+
+	TArray<FString> LodPathList;	
+	int32 CurrentLod = FCString::Atoi(*AssetImportData->AssetMetaInfo->ActiveLOD.Replace(TEXT("lod"), TEXT("")));
+
+	if (AssetImportData->AssetMetaInfo->ActiveLOD == TEXT("high")) CurrentLod = -1;
+	for (TSharedPtr<FAssetLodData> LodData : AssetImportData->LodList)
+	{		
+		if (LodData->Lod == TEXT("high")) continue;
+
+		int32 LodNumber = FCString::Atoi(*LodData->Lod.Replace(TEXT("lod"), TEXT("")));
+		if (LodNumber > CurrentLod)
+		{	
+			FString LodExtension, LodName;
+			LodData->Path.Split(TEXT("."), &LodName, &LodExtension);
+			LodExtension = FPaths::GetExtension(LodData->Path);
+			if (LodExtension == FileExtension)
+			{
+				
+				LodPathList.Add(LodData->Path);
+			}
+		}
+	}
+	return LodPathList;
+}
+
+
+TArray<TSharedPtr<FAssetLodData>> ParsePlantsLodList(TSharedPtr<FAssetTypeData> AssetImportData)
+{
+
+	FString FileExtension;
+	FileExtension = FPaths::GetExtension(AssetImportData->MeshList[0]->Path);
+	TArray<TSharedPtr<FAssetLodData>> SelectedLods;
+
+	
+	int32 CurrentLod = FCString::Atoi(*AssetImportData->AssetMetaInfo->ActiveLOD.Replace(TEXT("lod"), TEXT("")));
+
+	if (AssetImportData->AssetMetaInfo->ActiveLOD == TEXT("high")) CurrentLod = -1;
+	for (TSharedPtr<FAssetLodData> LodData : AssetImportData->LodList)
+	{
+		if (LodData->Lod == TEXT("high")) continue;
+
+		int32 LodNumber = FCString::Atoi(*LodData->Lod.Replace(TEXT("lod"), TEXT("")));
+		if (LodNumber > CurrentLod)
+		{
+			FString LodExtension, LodName;
+			LodData->Path.Split(TEXT("."), &LodName, &LodExtension);
+			LodExtension = FPaths::GetExtension(LodData->Path);
+			if (LodExtension == FileExtension)
+			{
+				SelectedLods.Add(LodData);
+			}
+		}
+	}
+
+	return SelectedLods;
+
+}
+
+
+FString GetRootDestination(const FString& ExportPath)
+{
+	TArray<FString> PathTokens;	
+	FString BasePath;
+	FString DestinationPath;
+	ExportPath.Split(TEXT("Content"), &BasePath, &DestinationPath);
+	FString RootDestination = TEXT("/Game/");
+
+#if PLATFORM_WINDOWS
+	DestinationPath.ParseIntoArray(PathTokens, TEXT("\\"));
+#elif PLATFORM_MAC
+	DestinationPath.ParseIntoArray(PathTokens, TEXT("/"));
+#elif PLATFORM_LINUX
+	DestinationPath.ParseIntoArray(PathTokens, TEXT("/"));
+#endif	
+
+	for (FString Token : PathTokens)
+	{
+		FString NormalizedToken = RemoveReservedKeywords( NormalizeString(Token));		
+		RootDestination = FPaths::Combine(RootDestination, NormalizedToken);
+	}	
+	if (RootDestination == TEXT("/Game/"))
+	{
+		RootDestination = FPaths::Combine(RootDestination, TEXT("Megascans"));
+	}
+	return RootDestination;
+}
+
+
+
+
+FString ResolveDestination(const FString& AssetDestination)
+{
+	return TEXT("Resolved destination path");
+}
+
+FString GetAssetName(TSharedPtr<FAssetTypeData> AssetImportData)
+{
+
+	return FString();
+}
+
+FString GetUniqueAssetName(const FString& AssetDestination, const FString AssetName, bool FileSearch)
+{
+
+	FString ResolvedAssetName = RemoveReservedKeywords(NormalizeString(AssetName));	
+	TArray<FString> AssetDirectories;	
+	auto& FileManager = IFileManager::Get();
+	FileManager.FindFiles(AssetDirectories, *AssetDestination, FileSearch, !FileSearch);
+
+	
+	for (int i = 0; i < 200; i++)
+	{
+		FString ResolvedAssetDirectory = TEXT("");
+		if (i < 10) {
+			ResolvedAssetDirectory = ResolvedAssetName + TEXT("_0") + FString::FromInt(i);
+		}
+		else
+		{
+			ResolvedAssetDirectory = ResolvedAssetName + TEXT("_") + FString::FromInt(i);
+		}
+
+		if (!UEditorAssetLibrary::DoesDirectoryExist(FPaths::Combine(AssetDestination, ResolvedAssetDirectory)))
+		{
+			
+			return ResolvedAssetDirectory;
+		}
+
+	}
+
+	return TEXT("");
+}
+
+FString SanitizeName(const FString& InputName)
+{
+	return RemoveReservedKeywords(NormalizeString(InputName));
+}
+
+
+
+
+	//template<typename T>
+	TArray<UMaterialInstanceConstant*> AssetUtils::GetSelectedAssets(const FString& AssetClass)
+	{
+		TArray<UMaterialInstanceConstant*> ObjectArray;
+
+		TArray<FAssetData> AssetDatas;
+		FContentBrowserModule& ContentBrowserModule = FModuleManager::LoadModuleChecked<FContentBrowserModule>("ContentBrowser");
+		IContentBrowserSingleton& ContentBrowserSingleton = ContentBrowserModule.Get();
+		ContentBrowserSingleton.GetSelectedAssets(AssetDatas);
+
+		for (FAssetData SelectedAsset : AssetDatas)
+		{
+			
+			if (SelectedAsset.AssetClass == FName(*AssetClass))
+			{
+				ObjectArray.Add(CastChecked<UMaterialInstanceConstant>(UEditorAssetLibrary::LoadAsset(SelectedAsset.ObjectPath.ToString())));
+			}
+		}
+		return ObjectArray;
+	}
+
+	void AssetUtils::FocusOnSelected(const FString& Path)
+	{
+		TArray<FString> Folders;
+		Folders.Add(Path);
+		FContentBrowserModule& ContentBrowserModule = FModuleManager::LoadModuleChecked<FContentBrowserModule>("ContentBrowser");
+		IContentBrowserSingleton& ContentBrowserSingleton = ContentBrowserModule.Get();
+		ContentBrowserSingleton.SyncBrowserToFolders(Folders);
+	}
+
+	
+	void AssetUtils::AddStaticMaterial(UStaticMesh* SourceMesh, UMaterialInstanceConstant* NewMaterial)
+	{
+		if (NewMaterial == nullptr) return;
+
+		UMaterialInterface* NewMatInterface = CastChecked<UMaterialInterface>(NewMaterial);
+		FStaticMaterial MaterialSlot = FStaticMaterial(NewMatInterface);
+		SourceMesh->StaticMaterials.Add(MaterialSlot);
+	}
+
+	void AssetUtils::SavePackage(UObject* SourceObject)
+	{
+		TArray<UObject*> InputObjects;
+		InputObjects.Add(SourceObject);
+		UPackageTools::SavePackagesForObjects(InputObjects);
+
+	}
+
+bool DHI::GetDHIJsonData(const FString & JsonStringData, TArray<FDHIData> & DHIAssetsData)
+{
+	FString StartString = TEXT("{\"DHIAssets\":");
+	FString EndString = TEXT("}");
+	FString FinalString = StartString + JsonStringData + EndString;
+	TSharedPtr<FJsonObject> ImportDataObject = DeserializeJson(FinalString);
+
+	TArray<TSharedPtr<FJsonValue> > AssetsImportDataArray = ImportDataObject->GetArrayField(TEXT("DHIAssets"));	
+
+	if (!AssetsImportDataArray[0]->AsObject()->HasField(TEXT("DHI"))) return false;	
+
+	for (TSharedPtr<FJsonValue> AssetDataObject : AssetsImportDataArray)
+	{
+		FDHIData CharacterData;
+		TSharedPtr<FJsonObject> CharacterDataObject = AssetDataObject->AsObject()->GetObjectField(TEXT("DHI"));
+		CharacterData.CharacterPath = CharacterDataObject->GetStringField(TEXT("characterPath"));
+		CharacterData.CharacterName = CharacterDataObject->GetStringField(TEXT("name"));
+		CharacterData.CommonPath = CharacterDataObject->GetStringField(TEXT("commonPath"));
+
+		CharacterData.CharacterPath = FPaths::Combine(CharacterData.CharacterPath, CharacterData.CharacterName);
+		DHIAssetsData.Add(CharacterData);
+		
+	}
+
+
+	return true;
+	
+}
+
+void DHI::CopyCharacter(const FDHIData & CharacterSourceData, bool OverWriteExisting)
+{	
+	
+
+	FString CommonDestinationPath = FPaths::ProjectContentDir();
+	FString MetaHumansRoot = FPaths::Combine(CommonDestinationPath, TEXT("MetaHumans"));
+	CommonDestinationPath = FPaths::Combine(MetaHumansRoot, TEXT("Common"));
+	FString RootFolder = TEXT("/Game/MetaHumans");
+	FString CommonFolder = FPaths::Combine(RootFolder, TEXT("Common"));
+	FString CharacterName = CharacterSourceData.CharacterName;
+	FString CharacterDestination = FPaths::Combine(MetaHumansRoot, CharacterName);	
+	IPlatformFile& PlatformFile = FPlatformFileManager::Get().GetPlatformFile();
+	if (!PlatformFile.DirectoryExists(*CommonDestinationPath))
+	{
+		PlatformFile.CreateDirectoryTree(*MetaHumansRoot);
+		PlatformFile.CreateDirectoryTree(*CommonDestinationPath);
+		FString CommonCopyMsg = TEXT("Copying Common folder.");
+		FText CommonCopyMsgDialogMessage = FText::FromString(CommonCopyMsg);
+		FScopedSlowTask AssetLoadprogress(1.0f, CommonCopyMsgDialogMessage, true);
+		AssetLoadprogress.MakeDialog();
+		AssetLoadprogress.EnterProgressFrame(1.0f);
+
+		PlatformFile.CopyDirectoryTree(*CommonDestinationPath, *CharacterSourceData.CommonPath, true);
+	}
+	
+	if (PlatformFile.DirectoryExists(*CharacterDestination) )
+	{
+		EAppReturnType::Type ContinueImport = FMessageDialog::Open(EAppMsgType::OkCancel, FText(FText::FromString("The character you are trying to import already exists. Do you want to overwrite it.")));
+		if (ContinueImport == EAppReturnType::Cancel) return;
+	}
+	PlatformFile.CreateDirectoryTree(*CharacterDestination);
+	FString CommonCopyMsg = TEXT("Copying : ") + CharacterName;
+	FText CommonCopyMsgDialogMessage = FText::FromString(CommonCopyMsg);
+	FScopedSlowTask AssetLoadprogress(1.0f, CommonCopyMsgDialogMessage, true);
+	AssetLoadprogress.MakeDialog();
+	AssetLoadprogress.EnterProgressFrame(1.0f);
+	PlatformFile.CopyDirectoryTree(*CharacterDestination, *CharacterSourceData.CharacterPath, true);
+	AssetUtils::FocusOnSelected(CharacterDestination);
+
+	
+	FAssetRegistryModule& AssetRegistryModule = FModuleManager::GetModuleChecked<FAssetRegistryModule>("AssetRegistry");
+	TArray<FString> AssetsBasePath;
+	FString BasePath = FPaths::Combine(FPaths::ProjectContentDir(), GetMSPresetsName());
+	AssetsBasePath.Add(TEXT("/Game/MetaHumans"));
+
+	AssetsBasePath.Add("/Game/MetaHumans/" + CharacterSourceData.CharacterName);
+	AssetRegistryModule.Get().ScanPathsSynchronous(AssetsBasePath, true);
+	
+}
+
+
+
+
+
diff --git a/Engine/Plugins/MegascansPlugin/Source/MegascansPlugin/Private/Utilities/MiscUtils.h b/Engine/Plugins/MegascansPlugin/Source/MegascansPlugin/Private/Utilities/MiscUtils.h
new file mode 100644
--- /dev/null
+++ b/Engine/Plugins/MegascansPlugin/Source/MegascansPlugin/Private/Utilities/MiscUtils.h
@@ -0,0 +1,56 @@
+// Copyright Epic Games, Inc. All Rights Reserved.
+#pragma once
+#include "CoreMinimal.h"
+#include "Runtime/Json/Public/Dom/JsonObject.h"
+#include "Runtime/Json/Public/Serialization/JsonSerializer.h"
+#include "Materials/MaterialInstanceConstant.h"
+#include "AssetImportData.h"
+
+class UStaticMesh;
+
+
+
+TSharedPtr<FJsonObject>  DeserializeJson(const FString& JsonStringData);
+void ShowErrorDialog(const FString& ErrorMessage);
+FString GetPluginPath();
+FString GetMaterialPresetsPath();
+FString GetPreferencesPath();
+FString GetMaterial(const FString & MaterialName);
+bool CopyMaterialPreset(const FString & MaterialName);
+FString GetMSPresetsName();
+UObject* LoadAsset(const FString& AssetPath);
+void DeleteExtraMesh(const FString& BasePath);
+void SaveAsset(const FString& AssetPath);
+TArray<FString> ParseLodList( TSharedPtr<FAssetTypeData> AssetImportData);
+TArray<TSharedPtr<FAssetLodData>> ParsePlantsLodList(TSharedPtr<FAssetTypeData> AssetImportData);
+FString ResolvePath(const FString& AssetPath, TSharedPtr<FAssetTypeData> AssetsImportData, TSharedPtr<FAssetTextureData> TextureData=nullptr);
+FString ResolveName(const FString& NamingConvention, TSharedPtr<FAssetTypeData> AssetsImportData, TSharedPtr<FAssetTextureData> TextureData = nullptr);
+FString RemoveReservedKeywords(const FString& Name);
+FString NormalizeString(FString InputString);
+TArray<FString> GetAssetsList(const FString& DirectoryPath);
+FString GetRootDestination(const FString& ExportPath);
+FString ResolveDestination(const FString& AssetDestination);
+bool CopyPresetTextures();
+
+FString GetAssetName(TSharedPtr<FAssetTypeData> AssetImportData);
+FString GetUniqueAssetName(const FString& AssetDestination, const FString AssetName, bool FileSearch=false);
+
+FString SanitizeName(const FString& InputName);
+
+namespace AssetUtils {
+	//template<typename T>
+	TArray<UMaterialInstanceConstant*> GetSelectedAssets(const FString& AssetClass);
+	void FocusOnSelected(const FString& Path);
+	void AddStaticMaterial(UStaticMesh* SourceMesh, UMaterialInstanceConstant* NewMaterial);
+	void SavePackage(UObject* SourceObject);
+}
+
+namespace PathUtils
+{
+	
+}
+
+namespace DHI {
+	bool GetDHIJsonData(const FString& JsonStringData, TArray<FDHIData>& DHIAssetsData);
+	void CopyCharacter(const FDHIData& CharacterSourcePath, bool OverWriteExisting = false);
+}
diff --git a/Engine/Plugins/MegascansPlugin/Source/MegascansPlugin/Public/IMegascansLiveLinkModule.h b/Engine/Plugins/MegascansPlugin/Source/MegascansPlugin/Public/IMegascansLiveLinkModule.h
new file mode 100644
--- /dev/null
+++ b/Engine/Plugins/MegascansPlugin/Source/MegascansPlugin/Public/IMegascansLiveLinkModule.h
@@ -0,0 +1,27 @@
+// Copyright Epic Games, Inc. All Rights Reserved.
+#pragma once
+
+#include "Modules/ModuleInterface.h"
+#include "Modules/ModuleManager.h"	
+
+
+
+#define MS_MODULE_NAME TEXT("MegascansPlugin")
+
+
+class IMegascansLiveLinkModule : public IModuleInterface
+{
+
+public:
+
+	static inline IMegascansLiveLinkModule& Get()
+	{
+		return FModuleManager::LoadModuleChecked<IMegascansLiveLinkModule>(MS_MODULE_NAME);
+	}
+
+	static inline bool IsAvailable()
+	{
+		return FModuleManager::Get().IsModuleLoaded(MS_MODULE_NAME);
+	}
+};
+
