Fix for "non-portable path to file".

diff --git a/Engine/Plugins/MegascansPlugin/Source/MegascansPlugin/Private/AssetImporters/Import3D.cpp b/Engine/Plugins/MegascansPlugin/Source/MegascansPlugin/Private/AssetImporters/Import3D.cpp
--- a/Engine/Plugins/MegascansPlugin/Source/MegascansPlugin/Private/AssetImporters/Import3D.cpp
+++ b/Engine/Plugins/MegascansPlugin/Source/MegascansPlugin/Private/AssetImporters/Import3D.cpp
@@ -4,7 +4,7 @@
 #include "AssetImporters/ImportSurface.h"
 
 
-#include "Runtime/Engine/Classes/Engine/StaticMesh.h"
+#include "Engine/StaticMesh.h"
 #include "Utilities/MeshOp.h"
 #include "Utilities/MiscUtils.h"
 #include "Utilities/MaterialUtils.h"
@@ -12,8 +12,8 @@
 #include "EditorAssetLibrary.h"
 #include "EditorStaticMeshLibrary.h"
 #include "UI/MSSettings.h"
-#include "Runtime/Core/Public/Misc/MessageDialog.h"
-#include "Runtime/Core/Public/Internationalization/Text.h"
+#include "Misc/MessageDialog.h"
+#include "Internationalization/Text.h"
 
 #include "StaticMeshEditorSubsystem.h"
 
diff --git a/Engine/Plugins/MegascansPlugin/Source/MegascansPlugin/Private/AssetImporters/Import3DPlant.cpp b/Engine/Plugins/MegascansPlugin/Source/MegascansPlugin/Private/AssetImporters/Import3DPlant.cpp
--- a/Engine/Plugins/MegascansPlugin/Source/MegascansPlugin/Private/AssetImporters/Import3DPlant.cpp
+++ b/Engine/Plugins/MegascansPlugin/Source/MegascansPlugin/Private/AssetImporters/Import3DPlant.cpp
@@ -6,7 +6,7 @@
 #include "Utilities/MeshOp.h"
 #include "AssetImportData.h"
 
-#include "Runtime/Engine/Classes/Engine/StaticMesh.h"
+#include "Engine/StaticMesh.h"
 #include "EditorAssetLibrary.h"
 #include "EditorStaticMeshLibrary.h"
 
diff --git a/Engine/Plugins/MegascansPlugin/Source/MegascansPlugin/Private/AssetImporters/ImportSurface.cpp b/Engine/Plugins/MegascansPlugin/Source/MegascansPlugin/Private/AssetImporters/ImportSurface.cpp
--- a/Engine/Plugins/MegascansPlugin/Source/MegascansPlugin/Private/AssetImporters/ImportSurface.cpp
+++ b/Engine/Plugins/MegascansPlugin/Source/MegascansPlugin/Private/AssetImporters/ImportSurface.cpp
@@ -6,9 +6,9 @@
 #include "AssetToolsModule.h"
 #include "IAssetTools.h"
 #include "AssetImportTask.h"
-#include "Editor/UnrealEd/Classes/Factories/MaterialInstanceConstantFactoryNew.h"
-#include "Runtime/Engine/Classes/Materials/MaterialInstanceConstant.h"
-#include "Runtime/Engine/Classes/Engine/Texture.h"
+#include "Factories/MaterialInstanceConstantFactoryNew.h"
+#include "Materials/MaterialInstanceConstant.h"
+#include "Engine/Texture.h"
 #include "UI/MSSettings.h"
 #include "MaterialEditingLibrary.h"
 #include "ContentBrowserModule.h"
@@ -16,7 +16,7 @@
 
 #include "EditorAssetLibrary.h"
 
-#include "Runtime/Engine/Classes/Engine/Selection.h"
+#include "Engine/Selection.h"
 #include "Editor.h"
 #include "PackageTools.h"
 
diff --git a/Engine/Plugins/MegascansPlugin/Source/MegascansPlugin/Private/AssetsImportController.cpp b/Engine/Plugins/MegascansPlugin/Source/MegascansPlugin/Private/AssetsImportController.cpp
--- a/Engine/Plugins/MegascansPlugin/Source/MegascansPlugin/Private/AssetsImportController.cpp
+++ b/Engine/Plugins/MegascansPlugin/Source/MegascansPlugin/Private/AssetsImportController.cpp
@@ -11,7 +11,7 @@
 #include "AssetImporters/ImportFactory.h"
 
 #include "Misc/MessageDialog.h"
-#include "Runtime/Core/Public/Internationalization/Text.h"
+#include "Internationalization/Text.h"
 #include "Misc/Paths.h"
 
 
diff --git a/Engine/Plugins/MegascansPlugin/Source/MegascansPlugin/Private/TCPServerImp.cpp b/Engine/Plugins/MegascansPlugin/Source/MegascansPlugin/Private/TCPServerImp.cpp
--- a/Engine/Plugins/MegascansPlugin/Source/MegascansPlugin/Private/TCPServerImp.cpp
+++ b/Engine/Plugins/MegascansPlugin/Source/MegascansPlugin/Private/TCPServerImp.cpp
@@ -2,8 +2,8 @@
 #include "TCPServerImp.h"
 #include <string>
 #include "NetworkMessage.h"
-
-#include "Runtime/Core/Public/Misc/ScopedSlowTask.h"
+#include "Async/Async.h"
+#include "Misc/ScopedSlowTask.h"
 
 
 FTCPServer::FTCPServer() 
diff --git a/Engine/Plugins/MegascansPlugin/Source/MegascansPlugin/Private/TCPServerImp.h b/Engine/Plugins/MegascansPlugin/Source/MegascansPlugin/Private/TCPServerImp.h
--- a/Engine/Plugins/MegascansPlugin/Source/MegascansPlugin/Private/TCPServerImp.h
+++ b/Engine/Plugins/MegascansPlugin/Source/MegascansPlugin/Private/TCPServerImp.h
@@ -4,9 +4,10 @@
 
 #include "CoreMinimal.h"
 #include <string>
-#include "Runtime/Sockets/Public/Sockets.h"
-#include "Runtime/Core/Public/Misc/DateTime.h"
-#include "Runtime/Core/Public/Async/Async.h"
+#include "Sockets.h"
+#include "Misc/DateTime.h"
+//#include "Runtime/Core/Public/Async/Async.h"
+#include "Async/AsyncWork.h"
 #include "MSPythonBridge.h"
 #include "UI/QMSUIManager.h"
 #include "AssetsImportController.h"
diff --git a/Engine/Plugins/MegascansPlugin/Source/MegascansPlugin/Private/Tools/BlendMaterials.cpp b/Engine/Plugins/MegascansPlugin/Source/MegascansPlugin/Private/Tools/BlendMaterials.cpp
--- a/Engine/Plugins/MegascansPlugin/Source/MegascansPlugin/Private/Tools/BlendMaterials.cpp
+++ b/Engine/Plugins/MegascansPlugin/Source/MegascansPlugin/Private/Tools/BlendMaterials.cpp
@@ -6,14 +6,14 @@
 #include "UI/MSSettings.h"
 #include "Utilities/Analytics.h"
 
-#include "Editor/UnrealEd/Classes/Factories/MaterialInstanceConstantFactoryNew.h"
-#include "Runtime/Engine/Classes/Materials/MaterialInstanceConstant.h"
-#include "Runtime/Engine/Classes/Engine/Texture.h"
-#include "Runtime/CoreUObject/Public/Misc/PackageName.h"
+#include "Factories/MaterialInstanceConstantFactoryNew.h"
+#include "Materials/MaterialInstanceConstant.h"
+#include "Engine/Texture.h"
+#include "Misc/PackageName.h"
 #include "MaterialEditingLibrary.h"
 
-#include "Runtime/Core/Public/Misc/MessageDialog.h"
-#include "Runtime/Core/Public/Internationalization/Text.h"
+#include "Misc/MessageDialog.h"
+#include "Internationalization/Text.h"
 
 
 TSharedPtr<FMaterialBlend> FMaterialBlend::MaterialBlendInst;
diff --git a/Engine/Plugins/MegascansPlugin/Source/MegascansPlugin/Private/UI/MSSettings.h b/Engine/Plugins/MegascansPlugin/Source/MegascansPlugin/Private/UI/MSSettings.h
--- a/Engine/Plugins/MegascansPlugin/Source/MegascansPlugin/Private/UI/MSSettings.h
+++ b/Engine/Plugins/MegascansPlugin/Source/MegascansPlugin/Private/UI/MSSettings.h
@@ -2,7 +2,7 @@
 #pragma once
 
 #include "CoreMinimal.h"
-#include "Runtime/Engine/Classes/Materials/Material.h"
+#include "Materials/Material.h"
 #include "UObject/ObjectMacros.h"
 #include "UObject/Object.h"
 #include "Engine/EngineTypes.h"
diff --git a/Engine/Plugins/MegascansPlugin/Source/MegascansPlugin/Private/UI/SMSWindow.cpp b/Engine/Plugins/MegascansPlugin/Source/MegascansPlugin/Private/UI/SMSWindow.cpp
--- a/Engine/Plugins/MegascansPlugin/Source/MegascansPlugin/Private/UI/SMSWindow.cpp
+++ b/Engine/Plugins/MegascansPlugin/Source/MegascansPlugin/Private/UI/SMSWindow.cpp
@@ -26,7 +26,7 @@
 #include "PropertyEditorModule.h"
 #include "Widgets/Docking/SDockTab.h"
 
-#include "Runtime/Engine/Classes/Materials/MaterialInterface.h"
+#include "Materials/MaterialInterface.h"
 #include "EditorAssetLibrary.h"
 
 
diff --git a/Engine/Plugins/MegascansPlugin/Source/MegascansPlugin/Private/Utilities/Analytics.cpp b/Engine/Plugins/MegascansPlugin/Source/MegascansPlugin/Private/Utilities/Analytics.cpp
--- a/Engine/Plugins/MegascansPlugin/Source/MegascansPlugin/Private/Utilities/Analytics.cpp
+++ b/Engine/Plugins/MegascansPlugin/Source/MegascansPlugin/Private/Utilities/Analytics.cpp
@@ -3,8 +3,8 @@
 #include "UI/MSSettings.h"
 #include "Interfaces/IHttpRequest.h"
 
-#include "Runtime/Json/Public/Dom/JsonObject.h"
-#include "Runtime/Json/Public/Serialization/JsonSerializer.h"
+#include "Dom/JsonObject.h"
+#include "Serialization/JsonSerializer.h"
 #include "HttpModule.h"
 #include "Misc/EngineVersion.h"
 
diff --git a/Engine/Plugins/MegascansPlugin/Source/MegascansPlugin/Private/Utilities/MaterialUtils.cpp b/Engine/Plugins/MegascansPlugin/Source/MegascansPlugin/Private/Utilities/MaterialUtils.cpp
--- a/Engine/Plugins/MegascansPlugin/Source/MegascansPlugin/Private/Utilities/MaterialUtils.cpp
+++ b/Engine/Plugins/MegascansPlugin/Source/MegascansPlugin/Private/Utilities/MaterialUtils.cpp
@@ -1,8 +1,8 @@
 // Copyright Epic Games, Inc. All Rights Reserved.
 #include "Utilities/MaterialUtils.h"
-#include "Editor/UnrealEd/Classes/Factories/MaterialInstanceConstantFactoryNew.h"
+#include "Factories/MaterialInstanceConstantFactoryNew.h"
 
-#include "Runtime/Engine/Classes/Engine/Texture.h"
+#include "Engine/Texture.h"
 #include "MaterialEditingLibrary.h"
 #include "AssetToolsModule.h"
 #include "IAssetTools.h"
diff --git a/Engine/Plugins/MegascansPlugin/Source/MegascansPlugin/Private/Utilities/MeshOp.cpp b/Engine/Plugins/MegascansPlugin/Source/MegascansPlugin/Private/Utilities/MeshOp.cpp
--- a/Engine/Plugins/MegascansPlugin/Source/MegascansPlugin/Private/Utilities/MeshOp.cpp
+++ b/Engine/Plugins/MegascansPlugin/Source/MegascansPlugin/Private/Utilities/MeshOp.cpp
@@ -1,23 +1,23 @@
 // Copyright Epic Games, Inc. All Rights Reserved.
 #include "Utilities/MeshOp.h"
-#include "UnrealEd/Classes/Factories/FbxImportUI.h"
-#include "UnrealEd/Classes/Factories/FbxStaticMeshImportData.h"
+#include "Factories/FbxImportUI.h"
+#include "Factories/FbxStaticMeshImportData.h"
 #include "AssetImportTask.h"
 #include "AssetToolsModule.h"
 #include "IAssetTools.h"
 #include "FbxMeshUtils.h"
 #include "AbcImportSettings.h"
 #include "Utilities/MiscUtils.h"
-#include "Runtime/Core/Public/Misc/Paths.h"
+#include "Misc/Paths.h"
 
 #include "EditorAssetLibrary.h"
 #include "EditorStaticMeshLibrary.h"
-#include "Editor/FoliageEdit/Private/FoliageTypeFactory.h"
-#include "Runtime/Foliage/Public/InstancedFoliageActor.h"
+//#include "Editor/FoliageEdit/Private/FoliageTypeFactory.h"
+#include "InstancedFoliageActor.h"
 
-#include "Runtime/Foliage/Public/FoliageType.h"
-#include "Runtime/Foliage/Public/FoliageType_InstancedStaticMesh.h"
-#include "Runtime/AssetRegistry/Public/AssetRegistryModule.h"
+#include "FoliageType.h"
+#include "FoliageType_InstancedStaticMesh.h"
+#include "AssetRegistryModule.h"
 #include "PerPlatformProperties.h"
 
 #include "StaticMeshEditorSubsystem.h"
diff --git a/Engine/Plugins/MegascansPlugin/Source/MegascansPlugin/Private/Utilities/MiscUtils.cpp b/Engine/Plugins/MegascansPlugin/Source/MegascansPlugin/Private/Utilities/MiscUtils.cpp
--- a/Engine/Plugins/MegascansPlugin/Source/MegascansPlugin/Private/Utilities/MiscUtils.cpp
+++ b/Engine/Plugins/MegascansPlugin/Source/MegascansPlugin/Private/Utilities/MiscUtils.cpp
@@ -1,24 +1,24 @@
 // Copyright Epic Games, Inc. All Rights Reserved.
 #include "Utilities/MiscUtils.h"
 
-#include "Runtime/Core/Public/Misc/Paths.h"
-#include "Runtime/Core/Public/HAL/PlatformFileManager.h"
-#include "Runtime/Core/Public/GenericPlatform/GenericPlatformFile.h"
-#include "Runtime/Core/Public/HAL/FileManager.h"
+#include "Misc/Paths.h"
+#include "HAL/PlatformFilemanager.h"
+//#include "Runtime/Core/Public/GenericPlatform/GenericPlatformFile.h"
+#include "HAL/FileManager.h"
 
 #include "EditorAssetLibrary.h"
-#include "Runtime/AssetRegistry/Public/AssetRegistryModule.h"
+#include "AssetRegistryModule.h"
 
 #include "ContentBrowserModule.h"
 #include "IContentBrowserSingleton.h"
-#include "Runtime/Engine/Classes/Engine/Selection.h"
+#include "Engine/Selection.h"
 
-#include "Editor/UnrealEd/Classes/Factories/MaterialInstanceConstantFactoryNew.h"
+#include "Factories/MaterialInstanceConstantFactoryNew.h"
 
-#include "Runtime/Engine/Classes/Engine/StaticMesh.h"
+#include "Engine/StaticMesh.h"
 
-#include "Runtime/Core/Public/Misc/MessageDialog.h"
-#include "Runtime/Core/Public/Internationalization/Text.h"
+#include "Misc/MessageDialog.h"
+#include "Internationalization/Text.h"
 
 #include "PackageTools.h"
 
diff --git a/Engine/Plugins/MegascansPlugin/Source/MegascansPlugin/Private/Utilities/MiscUtils.h b/Engine/Plugins/MegascansPlugin/Source/MegascansPlugin/Private/Utilities/MiscUtils.h
--- a/Engine/Plugins/MegascansPlugin/Source/MegascansPlugin/Private/Utilities/MiscUtils.h
+++ b/Engine/Plugins/MegascansPlugin/Source/MegascansPlugin/Private/Utilities/MiscUtils.h
@@ -1,8 +1,8 @@
 // Copyright Epic Games, Inc. All Rights Reserved.
 #pragma once
 #include "CoreMinimal.h"
-#include "Runtime/Json/Public/Dom/JsonObject.h"
-#include "Runtime/Json/Public/Serialization/JsonSerializer.h"
+#include "Dom/JsonObject.h"
+#include "Serialization/JsonSerializer.h"
 #include "Materials/MaterialInstanceConstant.h"
 #include "AssetImportData.h"
 
