Fixed "include text resolved to different locations". Renamed AssetData.h to MSAssetData.h.

diff --git a/Engine/Plugins/MegascansPlugin/Source/MegascansPlugin/Private/AssetImporters/Import3D.h b/Engine/Plugins/MegascansPlugin/Source/MegascansPlugin/Private/AssetImporters/Import3D.h
--- a/Engine/Plugins/MegascansPlugin/Source/MegascansPlugin/Private/AssetImporters/Import3D.h
+++ b/Engine/Plugins/MegascansPlugin/Source/MegascansPlugin/Private/AssetImporters/Import3D.h
@@ -4,7 +4,7 @@
 
 #include "AssetImportData.h"
 #include "ImportFactory.h"
-#include "Utilities/AssetData.h"
+#include "Utilities/MSAssetData.h"
 
 class UMaterialInstanceConstant;
 
diff --git a/Engine/Plugins/MegascansPlugin/Source/MegascansPlugin/Private/AssetImporters/Import3DPlant.h b/Engine/Plugins/MegascansPlugin/Source/MegascansPlugin/Private/AssetImporters/Import3DPlant.h
--- a/Engine/Plugins/MegascansPlugin/Source/MegascansPlugin/Private/AssetImporters/Import3DPlant.h
+++ b/Engine/Plugins/MegascansPlugin/Source/MegascansPlugin/Private/AssetImporters/Import3DPlant.h
@@ -3,7 +3,7 @@
 #include "CoreMinimal.h"
 
 #include "AssetImporters/ImportSurface.h"
-#include "Utilities/AssetData.h"
+#include "Utilities/MSAssetData.h"
 
 
 class UMaterialInstanceConstant;
diff --git a/Engine/Plugins/MegascansPlugin/Source/MegascansPlugin/Private/AssetImporters/ImportSurface.h b/Engine/Plugins/MegascansPlugin/Source/MegascansPlugin/Private/AssetImporters/ImportSurface.h
--- a/Engine/Plugins/MegascansPlugin/Source/MegascansPlugin/Private/AssetImporters/ImportSurface.h
+++ b/Engine/Plugins/MegascansPlugin/Source/MegascansPlugin/Private/AssetImporters/ImportSurface.h
@@ -4,7 +4,7 @@
 #include "AssetImportData.h"
 
 #include "ImportFactory.h"
-#include "Utilities/AssetData.h"
+#include "Utilities/MSAssetData.h"
 
 
 class UMaterialInstanceConstant;
diff --git a/Engine/Plugins/MegascansPlugin/Source/MegascansPlugin/Private/AssetsImportController.cpp b/Engine/Plugins/MegascansPlugin/Source/MegascansPlugin/Private/AssetsImportController.cpp
--- a/Engine/Plugins/MegascansPlugin/Source/MegascansPlugin/Private/AssetsImportController.cpp
+++ b/Engine/Plugins/MegascansPlugin/Source/MegascansPlugin/Private/AssetsImportController.cpp
@@ -7,7 +7,7 @@
 #include "UI/MSSettings.h"
 #include "Utilities/AssetsDatabase.h"
 #include "Utilities/Analytics.h"
-#include "Utilities/AssetData.h"
+#include "Utilities/MSAssetData.h"
 #include "AssetImporters/ImportFactory.h"
 
 #include "Misc/MessageDialog.h"
diff --git a/Engine/Plugins/MegascansPlugin/Source/MegascansPlugin/Private/Utilities/MSAssetData.cpp b/Engine/Plugins/MegascansPlugin/Source/MegascansPlugin/Private/Utilities/MSAssetData.cpp
new file mode 100644
--- /dev/null
+++ b/Engine/Plugins/MegascansPlugin/Source/MegascansPlugin/Private/Utilities/MSAssetData.cpp
@@ -0,0 +1,232 @@
+// Copyright Epic Games, Inc. All Rights Reserved.
+#include "MSAssetData.h"
+#include "Utilities/MiscUtils.h"
+#include "UI/MSSettings.h"
+#include "EditorAssetLibrary.h"
+
+
+TSharedPtr<FAssetImportParams> FAssetImportParams::ImportParamsInst;
+TSharedPtr<FAssetImportParams> FAssetImportParams::Get()
+{
+	if (!ImportParamsInst.IsValid())
+	{
+		ImportParamsInst = MakeShareable(new FAssetImportParams);
+	}
+	return ImportParamsInst;
+}
+
+
+TSharedPtr<ImportParams3DAsset> FAssetImportParams::Get3DAssetsParams(TSharedPtr<FAssetTypeData> AssetImportData)
+{
+	TSharedPtr<ImportParams3DAsset> AssetParams = MakeShareable(new ImportParams3DAsset);
+	AssetParams->BaseParams = GetAssetImportParams(AssetImportData);
+
+
+	const UMaterialAssetSettings* MatAssetPathsSettings = GetDefault<UMaterialAssetSettings>();
+	TSharedPtr<Asset3DParams>  Asset3DParameters = MakeShareable(new Asset3DParams);
+	Asset3DParameters->SubType = Get3DSubtype(AssetImportData);
+	Asset3DParameters->MeshDestination = AssetParams->BaseParams->AssetDestination;
+	Asset3DParameters->Asset3DMaterialType = GetS3DMasterMaterialType(AssetImportData);
+	
+	//TODO : create mapping of material with texture sets
+	TSharedPtr<SurfaceParams> MaterialParams = MakeShareable(new SurfaceParams);
+	if (Asset3DParameters->Asset3DMaterialType != ES3DMasterMaterialType::CUSTOM) {
+		MaterialParams->MasterMaterialName = MasterMaterialType[Asset3DParameters->Asset3DMaterialType];
+		MaterialParams->MasterMaterialPath = GetMasterMaterial(MaterialParams->MasterMaterialName);
+	}
+	else
+	{
+		MaterialParams->MasterMaterialName = TEXT("Custom");
+		MaterialParams->MasterMaterialPath = MatAssetPathsSettings->MasterMaterial3d;
+	}
+	
+	MaterialParams->MaterialInstanceDestination = AssetParams->BaseParams->AssetDestination;
+	MaterialParams->MaterialInstanceName = AssetParams->BaseParams->AssetName + TEXT("_inst");
+	MaterialParams->TexturesDestination = AssetParams->BaseParams->AssetDestination;
+	MaterialParams->bContainsPackedMaps = (AssetImportData->PackedTextures.Num() > 0) ? true : false;
+	Asset3DParameters->MaterialParams.Add(MaterialParams);
+
+	AssetParams->ParamsAssetType = Asset3DParameters;
+	
+	return AssetParams;
+}
+
+ES3DMasterMaterialType FAssetImportParams::GetS3DMasterMaterialType(TSharedPtr<FAssetTypeData> AssetImportData)
+{
+	const UMegascansSettings* MegascansSettings = GetDefault<UMegascansSettings>();	
+	const UMaterialAssetSettings* MatAssetPathsSettings = GetDefault<UMaterialAssetSettings>();
+	
+	if (AssetImportData->AssetMetaInfo->Type == TEXT("3d") && MatAssetPathsSettings->MasterMaterial3d != "None" && MatAssetPathsSettings->MasterMaterial3d != "")
+	{
+		if (UEditorAssetLibrary::DoesAssetExist(MatAssetPathsSettings->MasterMaterial3d))
+		{
+			return ES3DMasterMaterialType::CUSTOM;
+		}
+	}
+	if (AssetImportData->AssetMetaInfo->Type == TEXT("surface") && MatAssetPathsSettings->MasterMaterialSurface != "None" && MatAssetPathsSettings->MasterMaterialSurface != "")
+	{
+		if (UEditorAssetLibrary::DoesAssetExist(MatAssetPathsSettings->MasterMaterialSurface))
+		{
+			return ES3DMasterMaterialType::CUSTOM;
+		}
+
+	}
+	if (AssetImportData->PackedTextures.Num() > 0) {
+		return ES3DMasterMaterialType::CHANNEL_PACK;
+	}
+	if (MegascansSettings->bEnableDisplacement) { 
+		return ES3DMasterMaterialType::NORMAL_DISPLACEMENT; 
+	}
+	// FAssetTextureData type in auto
+	for (auto& TextureComponent : AssetImportData->TextureComponents)
+	{
+		if (TextureComponent->Type == TEXT("fuzz")) {
+			return ES3DMasterMaterialType::NORMAL_FUZZ;
+		}
+	}
+	return ES3DMasterMaterialType::NORMAL;
+}
+
+E3DPlantBillboardMaterial FAssetImportParams::GetBillboardMaterialType(TSharedPtr<FAssetTypeData> AssetImportData)
+{
+	// Type of master material for billboards
+	
+	return (AssetImportData->AssetMetaInfo->bUseBillboardMaterial) ? E3DPlantBillboardMaterial::CAMERA_FACING : E3DPlantBillboardMaterial::NORMAL_BILLBOARD;
+}
+
+
+TSharedPtr<ImportParamsSurfaceAsset> FAssetImportParams::GetSurfaceParams(TSharedPtr<FAssetTypeData> AssetImportData)
+{	
+	TSharedPtr< ImportParamsSurfaceAsset> AssetParams = MakeShareable(new ImportParamsSurfaceAsset);
+	AssetParams->BaseParams = GetAssetImportParams(AssetImportData);
+	const UMaterialAssetSettings* MatAssetPathsSettings = GetDefault<UMaterialAssetSettings>();
+	TSharedPtr<SurfaceParams> SurfaceParameters = MakeShareable( new SurfaceParams);
+	SurfaceParameters->SSubType = GetSurfaceSubtype(AssetImportData);
+	SurfaceParameters->MasterMaterialPath = TEXT("");
+	SurfaceParameters->MasterMaterialName = (SurfaceParameters->SSubType != ESurfaceSubType::SURFACE) ? SurfaceSubTypeMaterials[SurfaceParameters->SSubType] :MasterMaterialType[GetS3DMasterMaterialType(AssetImportData)];
+	if (SurfaceParameters->MasterMaterialName == TEXT("Custom"))
+	{		
+		SurfaceParameters->MasterMaterialPath = MatAssetPathsSettings->MasterMaterialSurface;		
+	}
+	else {
+		SurfaceParameters->MasterMaterialPath = GetMasterMaterial(SurfaceParameters->MasterMaterialName);
+	}		
+	SurfaceParameters->MaterialInstanceDestination = AssetParams->BaseParams->AssetDestination;	
+	SurfaceParameters->MaterialInstanceName = AssetParams->BaseParams->AssetName + TEXT("_inst");	
+	SurfaceParameters->TexturesDestination = AssetParams->BaseParams->AssetDestination;
+	SurfaceParameters->bContainsPackedMaps = (AssetImportData->PackedTextures.Num() > 0) ? true : false;	
+	AssetParams->ParamsAssetType = SurfaceParameters;
+	return AssetParams;
+}
+
+TSharedPtr<ImportParams3DPlantAsset> FAssetImportParams::Get3DPlantParams(TSharedPtr<FAssetTypeData> AssetImportData)
+{	
+	TSharedPtr<ImportParams3DPlantAsset>  AssetParams = MakeShareable(new ImportParams3DPlantAsset);
+	AssetParams->BaseParams = GetAssetImportParams(AssetImportData);
+	AssetParams->ParamsAssetType = MakeShareable(new Asset3DPlantParams);
+
+	const UMaterialAssetSettings* MatAssetPathsSettings = GetDefault<UMaterialAssetSettings>();		
+	AssetParams->ParamsAssetType->FoliageDestination = FPaths::Combine(AssetParams->BaseParams->AssetDestination, TEXT("Foliage"));
+	AssetParams->ParamsAssetType->MeshDestination = AssetParams->BaseParams->AssetDestination;
+	
+	//Plants main master material parameters
+	TSharedPtr < SurfaceParams> MasterMaterialParams = MakeShareable( new SurfaceParams);
+	MasterMaterialParams->MaterialInstanceDestination = AssetParams->BaseParams->AssetDestination;
+	MasterMaterialParams->MaterialInstanceName = AssetParams->BaseParams->AssetName + TEXT("_inst");
+	MasterMaterialParams->TexturesDestination = AssetParams->BaseParams->AssetDestination;
+	MasterMaterialParams->MasterMaterialName = PlantsMasterMaterial;
+	MasterMaterialParams->MasterMaterialPath = GetMasterMaterial(MasterMaterialParams->MasterMaterialName);
+	
+	if (MatAssetPathsSettings->MasterMaterialPlant != "None" && MatAssetPathsSettings->MasterMaterialPlant != "")
+	{
+		if (UEditorAssetLibrary::DoesAssetExist(MatAssetPathsSettings->MasterMaterialPlant))
+		{
+			
+			MasterMaterialParams->MasterMaterialName = TEXT("Custom");
+			MasterMaterialParams->MasterMaterialPath = MatAssetPathsSettings->MasterMaterialPlant;
+		}
+	}
+	//MasterMaterialParams->bContainsPackedMaps = (AssetImportData->PackedTextures.Num() > 0) ? true : false;
+	
+	//MasterMaterialParams->SSubType = ESurfaceSubType::SURFACE;
+	AssetParams->ParamsAssetType->PlantsMasterMaterial = MasterMaterialParams;	
+	
+	//Plants billboard master material parameters
+	TSharedPtr < SurfaceParams> BillboardMasterParams = MakeShareable( new SurfaceParams) ;
+	BillboardMasterParams->MasterMaterialName = BillboardMaterials[GetBillboardMaterialType(AssetImportData)];
+	BillboardMasterParams->MaterialInstanceDestination = FPaths::Combine(AssetParams->BaseParams->AssetDestination, TEXT("Billboard"));	
+	BillboardMasterParams->MaterialInstanceName = AssetParams->BaseParams->AssetName + TEXT("_Billboard_inst");
+	BillboardMasterParams->TexturesDestination = BillboardMasterParams->MaterialInstanceDestination;	
+	BillboardMasterParams->MasterMaterialPath = GetMasterMaterial(BillboardMasterParams->MasterMaterialName);	
+	AssetParams->ParamsAssetType->BillboardMasterMaterial = BillboardMasterParams;	
+	
+	return AssetParams;
+
+}
+
+TSharedPtr<AssetImportParams> FAssetImportParams::GetAssetImportParams(TSharedPtr<FAssetTypeData> AssetImportData)
+{
+	TSharedPtr<AssetImportParams>  AssetImportparameters = MakeShareable(new AssetImportParams);	
+	FString AssetType = AssetImportData->AssetMetaInfo->Type;
+	FString RootDestination = GetRootDestination(AssetImportData->AssetMetaInfo->ExportPath);
+	FString AssetTypePath = GetAssetTypePath(RootDestination, AssetImportData->AssetMetaInfo->Type);
+	AssetImportparameters->AssetName = GetUniqueAssetName(AssetTypePath, RemoveReservedKeywords(NormalizeString(AssetImportData->AssetMetaInfo->Name)));
+	AssetImportparameters->AssetDestination = FPaths::Combine(AssetTypePath, AssetImportparameters->AssetName);
+	return AssetImportparameters;
+}
+
+TSharedPtr<ImportParamsSurfaceAsset> FAssetImportParams::GetAtlasBrushParams(TSharedPtr<FAssetTypeData> AssetImportData)
+{
+	TSharedPtr< ImportParamsSurfaceAsset> AssetParams = MakeShareable(new ImportParamsSurfaceAsset);
+	AssetParams->BaseParams = GetAssetImportParams(AssetImportData);	
+	TSharedPtr<SurfaceParams> SurfaceParameters = MakeShareable(new SurfaceParams);
+	SurfaceParameters->MasterMaterialName = (AssetImportData->AssetMetaInfo->Type == TEXT("atlas")) ? AtlasSubTypeMaterials[GetAtlasSubtype(AssetImportData)] : BrushMasterMaterial;
+	SurfaceParameters->MasterMaterialPath = GetMasterMaterial(SurfaceParameters->MasterMaterialName);	
+	SurfaceParameters->MaterialInstanceDestination = AssetParams->BaseParams->AssetDestination;
+	SurfaceParameters->MaterialInstanceName = AssetParams->BaseParams->AssetName + TEXT("_inst");
+	SurfaceParameters->TexturesDestination = AssetParams->BaseParams->AssetDestination;
+	SurfaceParameters->bContainsPackedMaps = false;
+
+	AssetParams->ParamsAssetType = SurfaceParameters;
+	return AssetParams;
+}
+
+FString FAssetImportParams::GetAssetTypePath(const FString& RootDestination, const FString& AssetType)
+{
+	TMap<FString, FString> TypeNames = {
+		{TEXT("3d"),TEXT("3D_Assets")},
+		{TEXT("3dplant"),TEXT("3D_Plants")},
+		{TEXT("surface"),TEXT("Surfaces")},
+		{TEXT("atlas"),TEXT("Atlases")},
+		{TEXT("brush"),TEXT("Brushes")},
+		{TEXT("3datlas"),TEXT("3D_Atlases")}
+	};
+	return FPaths::Combine(RootDestination, TypeNames[AssetType]);
+}
+
+EAtlasSubType FAssetImportParams::GetAtlasSubtype(TSharedPtr<FAssetTypeData> AssetImportData)
+{
+	auto& TagsList = AssetImportData->AssetMetaInfo->Tags;
+	return (!TagsList.Contains(TEXT("decal"))) ? EAtlasSubType::ATLAS: EAtlasSubType::DECAL;
+}
+ESurfaceSubType FAssetImportParams::GetSurfaceSubtype(TSharedPtr<FAssetTypeData> AssetImportData)
+{
+	auto& TagsList = AssetImportData->AssetMetaInfo->Tags;
+	return (TagsList.Contains(TEXT("metal"))) ? ESurfaceSubType::METAL : (TagsList.Contains(TEXT("surface imperfection"))) ? ESurfaceSubType::IMPERFECTION : (TagsList.Contains(TEXT("tileable displacement"))) ? ESurfaceSubType::DISPLACEMENT : ESurfaceSubType::SURFACE;
+}
+
+EAsset3DSubtype FAssetImportParams::Get3DSubtype(TSharedPtr<FAssetTypeData> AssetImportData)
+{
+	return (AssetImportData->AssetMetaInfo->Categories.Contains(TEXT("scatter")) || AssetImportData->AssetMetaInfo->Tags.Contains(TEXT("scatter")) || AssetImportData->AssetMetaInfo->Categories.Contains(TEXT("cmb_asset")) || AssetImportData->AssetMetaInfo->Tags.Contains(TEXT("cmb_asset"))) ? EAsset3DSubtype::MULTI_MESH : EAsset3DSubtype::NORMAL_3D;
+	
+}
+
+FString FAssetImportParams::GetMasterMaterial(const FString& SelectedMaterial)
+{
+	CopyPresetTextures();
+	FString MaterialPath = GetMaterial(SelectedMaterial);
+	return (MaterialPath == TEXT("")) ? TEXT("") : MaterialPath;
+}
+
+
+
diff --git a/Engine/Plugins/MegascansPlugin/Source/MegascansPlugin/Private/Utilities/MSAssetData.h b/Engine/Plugins/MegascansPlugin/Source/MegascansPlugin/Private/Utilities/MSAssetData.h
new file mode 100644
--- /dev/null
+++ b/Engine/Plugins/MegascansPlugin/Source/MegascansPlugin/Private/Utilities/MSAssetData.h
@@ -0,0 +1,160 @@
+// Copyright Epic Games, Inc. All Rights Reserved.
+#pragma once
+#include "CoreMinimal.h"
+#include "AssetImportData.h"
+
+
+enum EAtlasSubType {
+	DECAL,
+	ATLAS
+};
+
+enum ESurfaceSubType {
+	METAL,
+	IMPERFECTION,
+	DISPLACEMENT,
+	SURFACE
+};
+
+enum EAsset3DSubtype {
+	MULTI_MESH,
+	NORMAL_3D
+};
+
+enum ES3DMasterMaterialType {
+	NORMAL,
+	NORMAL_FUZZ,
+	CHANNEL_PACK,
+	NORMAL_DISPLACEMENT,
+	MODULAR_WINDOWS,
+	CUSTOM
+};
+
+enum E3DPlantBillboardMaterial
+{
+	NORMAL_BILLBOARD,
+	CAMERA_FACING
+};
+
+
+struct SurfaceParams
+{
+	FString MaterialInstanceName;
+	FString MasterMaterialName;
+	FString MasterMaterialPath;
+	FString MaterialInstanceDestination;
+	FString TexturesDestination;
+	ESurfaceSubType SSubType;
+	bool bContainsPackedMaps;
+};
+
+struct Asset3DParams
+{
+	FString MeshDestination;
+	ES3DMasterMaterialType Asset3DMaterialType;
+	EAsset3DSubtype SubType;	
+	TArray<TSharedPtr<SurfaceParams>> MaterialParams;
+};
+
+struct Asset3DPlantParams
+{
+	FString MeshDestination;
+	FString FoliageDestination;
+	E3DPlantBillboardMaterial BillboardMaterialType;
+	TSharedPtr<SurfaceParams> BillboardMasterMaterial;
+	TSharedPtr<SurfaceParams> PlantsMasterMaterial;
+	
+};
+
+struct AssetImportParams
+{	
+	FString AssetName;	
+	FString AssetDestination;	
+};
+
+
+// Main data structures
+struct ImportParams3DAsset
+{
+	TSharedPtr<AssetImportParams> BaseParams;
+	TSharedPtr<Asset3DParams> ParamsAssetType;
+};
+
+struct ImportParams3DPlantAsset
+{
+	TSharedPtr<AssetImportParams> BaseParams;
+	TSharedPtr<Asset3DPlantParams> ParamsAssetType;
+};
+
+struct ImportParamsSurfaceAsset
+{
+	TSharedPtr<AssetImportParams> BaseParams;
+	TSharedPtr<SurfaceParams> ParamsAssetType;
+};
+
+
+
+
+
+class FAssetImportParams {
+private:
+	FAssetImportParams() = default;
+
+	const TMap<EAtlasSubType, FString> AtlasSubTypeMaterials = {
+		{EAtlasSubType::ATLAS,"Atlas_Material"},
+		{EAtlasSubType::DECAL,"Decal_MasterMaterial"}
+	};
+
+	const TMap<E3DPlantBillboardMaterial, FString> BillboardMaterials = {
+		{E3DPlantBillboardMaterial::CAMERA_FACING,"Billboard_Material"},
+		{E3DPlantBillboardMaterial::NORMAL_BILLBOARD,"Foliage_Material"}
+	};
+	//These values should be stored in a configuration file
+	const TMap<ESurfaceSubType, FString> SurfaceSubTypeMaterials = { 
+		{ESurfaceSubType::METAL,"MS_DefaultMaterial"},
+		{ESurfaceSubType::IMPERFECTION,"Imperfection_MasterMaterial"} ,
+		{ESurfaceSubType::DISPLACEMENT,"Displacement_MasterMaterial"}
+	};
+
+	const FString BrushMasterMaterial = TEXT("BrushDecal_Material");
+
+	const TMap<ES3DMasterMaterialType, FString> MasterMaterialType = {
+		{ES3DMasterMaterialType::CHANNEL_PACK , TEXT("MS_DefaultMaterial_CP")},
+		{ES3DMasterMaterialType::MODULAR_WINDOWS , TEXT("")},
+		{ES3DMasterMaterialType::NORMAL , TEXT("MS_DefaultMaterial")},
+		{ES3DMasterMaterialType::NORMAL_DISPLACEMENT , TEXT("MS_DefaultMaterial_Displacement")},
+		{ES3DMasterMaterialType::NORMAL_FUZZ , TEXT("MS_DefaultMaterial_Fuzzy")},
+		{ES3DMasterMaterialType::CUSTOM , TEXT("Custom")}
+	};
+
+	const FString PlantsMasterMaterial = TEXT("Foliage_Material");
+	
+	static TSharedPtr<FAssetImportParams> ImportParamsInst;
+	FString GetAssetTypePath(const FString& RootDestination, const FString& AssetType);
+	EAtlasSubType GetAtlasSubtype(TSharedPtr<FAssetTypeData> AssetImportData);
+	ESurfaceSubType GetSurfaceSubtype(TSharedPtr<FAssetTypeData> AssetImportData);
+	EAsset3DSubtype Get3DSubtype(TSharedPtr<FAssetTypeData> AssetImportData);
+	ES3DMasterMaterialType GetS3DMasterMaterialType(TSharedPtr<FAssetTypeData> AssetImportData);
+	E3DPlantBillboardMaterial GetBillboardMaterialType(TSharedPtr<FAssetTypeData> AssetImportData);
+	FString GetMasterMaterial(const FString& SelectedMaterial);
+
+public:
+	static TSharedPtr<FAssetImportParams> Get();	
+//Asset related data
+	TSharedPtr<ImportParams3DAsset> Get3DAssetsParams(TSharedPtr<FAssetTypeData> AssetImportData);
+	TSharedPtr<ImportParamsSurfaceAsset> GetSurfaceParams(TSharedPtr<FAssetTypeData> AssetImportData);
+	TSharedPtr<ImportParams3DPlantAsset> Get3DPlantParams(TSharedPtr<FAssetTypeData> AssetImportData);
+	TSharedPtr<AssetImportParams> GetAssetImportParams(TSharedPtr<FAssetTypeData> AssetImportData);
+	TSharedPtr<ImportParamsSurfaceAsset> GetAtlasBrushParams(TSharedPtr<FAssetTypeData> AssetImportData);
+};
+
+
+
+
+
+
+
+
+
+
+
