Fix for nounity errors in MegascansPlugin

diff --git a/Engine/Plugins/MegascansPlugin/Source/MegascansPlugin/Private/AssetImporters/Import3DPlant.h b/Engine/Plugins/MegascansPlugin/Source/MegascansPlugin/Private/AssetImporters/Import3DPlant.h
--- a/Engine/Plugins/MegascansPlugin/Source/MegascansPlugin/Private/AssetImporters/Import3DPlant.h
+++ b/Engine/Plugins/MegascansPlugin/Source/MegascansPlugin/Private/AssetImporters/Import3DPlant.h
@@ -7,6 +7,7 @@
 
 
 class UMaterialInstanceConstant;
+class UStaticMesh;
 
 enum PlantImportType {
 	BILLBOARD_ONLY,
diff --git a/Engine/Plugins/MegascansPlugin/Source/MegascansPlugin/Private/AssetImporters/ImportSurface.h b/Engine/Plugins/MegascansPlugin/Source/MegascansPlugin/Private/AssetImporters/ImportSurface.h
--- a/Engine/Plugins/MegascansPlugin/Source/MegascansPlugin/Private/AssetImporters/ImportSurface.h
+++ b/Engine/Plugins/MegascansPlugin/Source/MegascansPlugin/Private/AssetImporters/ImportSurface.h
@@ -30,6 +30,8 @@ struct SurfaceImportParams
 	bool bExrEnabled;
 };
 
+// Forward declare
+class UAssetImportTask;
 
 class FImportSurface : public IImportAsset
 {
diff --git a/Engine/Plugins/MegascansPlugin/Source/MegascansPlugin/Private/MSPythonBridge.cpp b/Engine/Plugins/MegascansPlugin/Source/MegascansPlugin/Private/MSPythonBridge.cpp
--- a/Engine/Plugins/MegascansPlugin/Source/MegascansPlugin/Private/MSPythonBridge.cpp
+++ b/Engine/Plugins/MegascansPlugin/Source/MegascansPlugin/Private/MSPythonBridge.cpp
@@ -1,5 +1,6 @@
 // Copyright Epic Games, Inc. All Rights Reserved.
 #include "MSPythonBridge.h"
+#include "UObject/UObjectHash.h"
 
 UMSPythonBridge* UMSPythonBridge::Get()
 {
diff --git a/Engine/Plugins/MegascansPlugin/Source/MegascansPlugin/Private/TCPServerImp.h b/Engine/Plugins/MegascansPlugin/Source/MegascansPlugin/Private/TCPServerImp.h
--- a/Engine/Plugins/MegascansPlugin/Source/MegascansPlugin/Private/TCPServerImp.h
+++ b/Engine/Plugins/MegascansPlugin/Source/MegascansPlugin/Private/TCPServerImp.h
@@ -17,6 +17,7 @@
 #include "Common/TcpListener.h"
 #include "Sockets.h"
 #include "Common/UdpSocketReceiver.h"
+#include "Containers/Queue.h"
 
 
 
diff --git a/Engine/Plugins/MegascansPlugin/Source/MegascansPlugin/Private/UI/MaterialBlendingDetails.cpp b/Engine/Plugins/MegascansPlugin/Source/MegascansPlugin/Private/UI/MaterialBlendingDetails.cpp
--- a/Engine/Plugins/MegascansPlugin/Source/MegascansPlugin/Private/UI/MaterialBlendingDetails.cpp
+++ b/Engine/Plugins/MegascansPlugin/Source/MegascansPlugin/Private/UI/MaterialBlendingDetails.cpp
@@ -7,6 +7,7 @@
 #include "DetailWidgetRow.h"
 
 #include "Tools/BlendMaterials.h"
+#include "Widgets/Input/SButton.h"
 
 #define LOCTEXT_NAMESPACE "MegascansLiveLinkModule"
 
diff --git a/Engine/Plugins/MegascansPlugin/Source/MegascansPlugin/Private/UI/MaterialBlendingDetails.h b/Engine/Plugins/MegascansPlugin/Source/MegascansPlugin/Private/UI/MaterialBlendingDetails.h
--- a/Engine/Plugins/MegascansPlugin/Source/MegascansPlugin/Private/UI/MaterialBlendingDetails.h
+++ b/Engine/Plugins/MegascansPlugin/Source/MegascansPlugin/Private/UI/MaterialBlendingDetails.h
@@ -1,6 +1,7 @@
 // Copyright Epic Games, Inc. All Rights Reserved.
 #pragma once
 #include "IDetailCustomization.h"
+#include "Input/Reply.h"
 
 class BlendSettingsCustomization : public IDetailCustomization
 
diff --git a/Engine/Plugins/MegascansPlugin/Source/MegascansPlugin/Private/Utilities/Analytics.cpp b/Engine/Plugins/MegascansPlugin/Source/MegascansPlugin/Private/Utilities/Analytics.cpp
--- a/Engine/Plugins/MegascansPlugin/Source/MegascansPlugin/Private/Utilities/Analytics.cpp
+++ b/Engine/Plugins/MegascansPlugin/Source/MegascansPlugin/Private/Utilities/Analytics.cpp
@@ -7,6 +7,7 @@
 #include "Serialization/JsonSerializer.h"
 #include "HttpModule.h"
 #include "Misc/EngineVersion.h"
+#include "Launch/Resources/Version.h"
 
 TSharedPtr<FAnalytics> FAnalytics::AnalyticsInst;
 
diff --git a/Engine/Plugins/MegascansPlugin/Source/MegascansPlugin/Private/Utilities/Analytics.h b/Engine/Plugins/MegascansPlugin/Source/MegascansPlugin/Private/Utilities/Analytics.h
--- a/Engine/Plugins/MegascansPlugin/Source/MegascansPlugin/Private/Utilities/Analytics.h
+++ b/Engine/Plugins/MegascansPlugin/Source/MegascansPlugin/Private/Utilities/Analytics.h
@@ -2,8 +2,8 @@
 #pragma once
 #include "CoreMinimal.h"
 
-
-
+// Forward declare
+class FJsonObject;
 
 class FAnalytics
 {
