Changed FAssetImportParams to FMSImportParams to avoid ambiguity with existing class.

diff --git a/Engine/Plugins/MegascansPlugin/Source/MegascansPlugin/Private/AssetImporters/Import3D.cpp b/Engine/Plugins/MegascansPlugin/Source/MegascansPlugin/Private/AssetImporters/Import3D.cpp
--- a/Engine/Plugins/MegascansPlugin/Source/MegascansPlugin/Private/AssetImporters/Import3D.cpp
+++ b/Engine/Plugins/MegascansPlugin/Source/MegascansPlugin/Private/AssetImporters/Import3D.cpp
@@ -32,7 +32,7 @@ TSharedPtr<FImport3d> FImport3d::Get()
 void FImport3d::ImportAsset(TSharedPtr<FAssetTypeData> AssetImportData)
 {	
 	const UMegascansSettings* MegascansSettings = GetDefault<UMegascansSettings>();
-	TSharedPtr<FAssetImportParams> AssetSetupParameters = FAssetImportParams::Get();
+	TSharedPtr<FMSImportParams> AssetSetupParameters = FMSImportParams::Get();
 	TSharedPtr<ImportParams3DAsset> Asset3DParameters = AssetSetupParameters->Get3DAssetsParams(AssetImportData);	
 
 	if (AssetImportData->MeshList.Num() == 0) return;
diff --git a/Engine/Plugins/MegascansPlugin/Source/MegascansPlugin/Private/AssetImporters/Import3DPlant.cpp b/Engine/Plugins/MegascansPlugin/Source/MegascansPlugin/Private/AssetImporters/Import3DPlant.cpp
--- a/Engine/Plugins/MegascansPlugin/Source/MegascansPlugin/Private/AssetImporters/Import3DPlant.cpp
+++ b/Engine/Plugins/MegascansPlugin/Source/MegascansPlugin/Private/AssetImporters/Import3DPlant.cpp
@@ -32,7 +32,7 @@ void FImportPlant::ImportAsset(TSharedPtr<FAssetTypeData> AssetImportData)
 {
 	if (AssetImportData->MeshList.Num() == 0) return;
 	const UMegascansSettings* MegascansSettings = GetDefault<UMegascansSettings>();
-	TSharedPtr<FAssetImportParams> AssetSetupParameters = FAssetImportParams::Get();
+	TSharedPtr<FMSImportParams> AssetSetupParameters = FMSImportParams::Get();
 	TSharedPtr<ImportParams3DPlantAsset> AssetPlantParameters = AssetSetupParameters->Get3DPlantParams(AssetImportData);	
 	PlantImportType ImportType = GetImportType(AssetImportData);
 	TMap<FString, FString> ImportedPlants = ImportPlants(AssetImportData, AssetPlantParameters);
diff --git a/Engine/Plugins/MegascansPlugin/Source/MegascansPlugin/Private/AssetImporters/ImportSurface.cpp b/Engine/Plugins/MegascansPlugin/Source/MegascansPlugin/Private/AssetImporters/ImportSurface.cpp
--- a/Engine/Plugins/MegascansPlugin/Source/MegascansPlugin/Private/AssetImporters/ImportSurface.cpp
+++ b/Engine/Plugins/MegascansPlugin/Source/MegascansPlugin/Private/AssetImporters/ImportSurface.cpp
@@ -24,7 +24,7 @@ TSharedPtr<FImportSurface> FImportSurface::ImportSurfaceInst;
 
 void FImportSurface::ImportAsset(TSharedPtr<FAssetTypeData> AssetImportData)
 {
-	TSharedPtr<FAssetImportParams> AssetSetupParameters = FAssetImportParams::Get();	
+	TSharedPtr<FMSImportParams> AssetSetupParameters = FMSImportParams::Get();	
 	TSharedPtr<ImportParamsSurfaceAsset> SurfaceImportParams = (AssetImportData->AssetMetaInfo->Type == TEXT("surface")) ? AssetSetupParameters->GetSurfaceParams(AssetImportData) : AssetSetupParameters->GetAtlasBrushParams(AssetImportData);
 	UMaterialInstanceConstant* MaterialInstance = ImportSurface(AssetImportData, SurfaceImportParams->ParamsAssetType);
 	AssetUtils::FocusOnSelected(SurfaceImportParams->ParamsAssetType->MaterialInstanceDestination);
diff --git a/Engine/Plugins/MegascansPlugin/Source/MegascansPlugin/Private/AssetsImportController.cpp b/Engine/Plugins/MegascansPlugin/Source/MegascansPlugin/Private/AssetsImportController.cpp
--- a/Engine/Plugins/MegascansPlugin/Source/MegascansPlugin/Private/AssetsImportController.cpp
+++ b/Engine/Plugins/MegascansPlugin/Source/MegascansPlugin/Private/AssetsImportController.cpp
@@ -75,8 +75,8 @@ void FAssetsImportController::ImportAssets(const FString& AssetsImportJson)
 	for (TSharedPtr<FAssetTypeData> AssetImportData : AssetsImportData->AllAssetsData)
 	{	
 		//AssetImportData->AssetMetaInfo->bSavePackages = bSavePackages;
-		AssetImportData->AssetMetaInfo->bSavePackages = false;
-		TSharedPtr<FAssetImportParams> AssetSetupParameters = FAssetImportParams::Get();
+		AssetImportData->AssetMetaInfo->bSavePackages = true;
+		TSharedPtr<FMSImportParams> AssetSetupParameters = FMSImportParams::Get();
 		AssetRecord Record;
 		if (FAssetsDatabase::Get()->RecordExists(AssetImportData->AssetMetaInfo->Id, Record) && FPaths::DirectoryExists(FPaths::Combine(FPaths::ProjectContentDir(), Record.Path.Replace(TEXT("/Game"), TEXT("")))))
 		{
diff --git a/Engine/Plugins/MegascansPlugin/Source/MegascansPlugin/Private/Utilities/MSAssetData.cpp b/Engine/Plugins/MegascansPlugin/Source/MegascansPlugin/Private/Utilities/MSAssetData.cpp
--- a/Engine/Plugins/MegascansPlugin/Source/MegascansPlugin/Private/Utilities/MSAssetData.cpp
+++ b/Engine/Plugins/MegascansPlugin/Source/MegascansPlugin/Private/Utilities/MSAssetData.cpp
@@ -5,18 +5,18 @@
 #include "EditorAssetLibrary.h"
 
 
-TSharedPtr<FAssetImportParams> FAssetImportParams::ImportParamsInst;
-TSharedPtr<FAssetImportParams> FAssetImportParams::Get()
+TSharedPtr<FMSImportParams> FMSImportParams::ImportParamsInst;
+TSharedPtr<FMSImportParams> FMSImportParams::Get()
 {
 	if (!ImportParamsInst.IsValid())
 	{
-		ImportParamsInst = MakeShareable(new FAssetImportParams);
+		ImportParamsInst = MakeShareable(new FMSImportParams);
 	}
 	return ImportParamsInst;
 }
 
 
-TSharedPtr<ImportParams3DAsset> FAssetImportParams::Get3DAssetsParams(TSharedPtr<FAssetTypeData> AssetImportData)
+TSharedPtr<ImportParams3DAsset> FMSImportParams::Get3DAssetsParams(TSharedPtr<FAssetTypeData> AssetImportData)
 {
 	TSharedPtr<ImportParams3DAsset> AssetParams = MakeShareable(new ImportParams3DAsset);
 	AssetParams->BaseParams = GetAssetImportParams(AssetImportData);
@@ -51,7 +51,7 @@ TSharedPtr<ImportParams3DAsset> FAssetImportParams::Get3DAssetsParams(TSharedPtr
 	return AssetParams;
 }
 
-ES3DMasterMaterialType FAssetImportParams::GetS3DMasterMaterialType(TSharedPtr<FAssetTypeData> AssetImportData)
+ES3DMasterMaterialType FMSImportParams::GetS3DMasterMaterialType(TSharedPtr<FAssetTypeData> AssetImportData)
 {
 	const UMegascansSettings* MegascansSettings = GetDefault<UMegascansSettings>();	
 	const UMaterialAssetSettings* MatAssetPathsSettings = GetDefault<UMaterialAssetSettings>();
@@ -87,7 +87,7 @@ ES3DMasterMaterialType FAssetImportParams::GetS3DMasterMaterialType(TSharedPtr<F
 	return ES3DMasterMaterialType::NORMAL;
 }
 
-E3DPlantBillboardMaterial FAssetImportParams::GetBillboardMaterialType(TSharedPtr<FAssetTypeData> AssetImportData)
+E3DPlantBillboardMaterial FMSImportParams::GetBillboardMaterialType(TSharedPtr<FAssetTypeData> AssetImportData)
 {
 	// Type of master material for billboards
 	
@@ -95,7 +95,7 @@ E3DPlantBillboardMaterial FAssetImportParams::GetBillboardMaterialType(TSharedPt
 }
 
 
-TSharedPtr<ImportParamsSurfaceAsset> FAssetImportParams::GetSurfaceParams(TSharedPtr<FAssetTypeData> AssetImportData)
+TSharedPtr<ImportParamsSurfaceAsset> FMSImportParams::GetSurfaceParams(TSharedPtr<FAssetTypeData> AssetImportData)
 {	
 	TSharedPtr< ImportParamsSurfaceAsset> AssetParams = MakeShareable(new ImportParamsSurfaceAsset);
 	AssetParams->BaseParams = GetAssetImportParams(AssetImportData);
@@ -119,7 +119,7 @@ TSharedPtr<ImportParamsSurfaceAsset> FAssetImportParams::GetSurfaceParams(TShare
 	return AssetParams;
 }
 
-TSharedPtr<ImportParams3DPlantAsset> FAssetImportParams::Get3DPlantParams(TSharedPtr<FAssetTypeData> AssetImportData)
+TSharedPtr<ImportParams3DPlantAsset> FMSImportParams::Get3DPlantParams(TSharedPtr<FAssetTypeData> AssetImportData)
 {	
 	TSharedPtr<ImportParams3DPlantAsset>  AssetParams = MakeShareable(new ImportParams3DPlantAsset);
 	AssetParams->BaseParams = GetAssetImportParams(AssetImportData);
@@ -164,7 +164,7 @@ TSharedPtr<ImportParams3DPlantAsset> FAssetImportParams::Get3DPlantParams(TShare
 
 }
 
-TSharedPtr<AssetImportParams> FAssetImportParams::GetAssetImportParams(TSharedPtr<FAssetTypeData> AssetImportData)
+TSharedPtr<AssetImportParams> FMSImportParams::GetAssetImportParams(TSharedPtr<FAssetTypeData> AssetImportData)
 {
 	TSharedPtr<AssetImportParams>  AssetImportparameters = MakeShareable(new AssetImportParams);	
 	FString AssetType = AssetImportData->AssetMetaInfo->Type;
@@ -175,7 +175,7 @@ TSharedPtr<AssetImportParams> FAssetImportParams::GetAssetImportParams(TSharedPt
 	return AssetImportparameters;
 }
 
-TSharedPtr<ImportParamsSurfaceAsset> FAssetImportParams::GetAtlasBrushParams(TSharedPtr<FAssetTypeData> AssetImportData)
+TSharedPtr<ImportParamsSurfaceAsset> FMSImportParams::GetAtlasBrushParams(TSharedPtr<FAssetTypeData> AssetImportData)
 {
 	TSharedPtr< ImportParamsSurfaceAsset> AssetParams = MakeShareable(new ImportParamsSurfaceAsset);
 	AssetParams->BaseParams = GetAssetImportParams(AssetImportData);	
@@ -191,7 +191,7 @@ TSharedPtr<ImportParamsSurfaceAsset> FAssetImportParams::GetAtlasBrushParams(TSh
 	return AssetParams;
 }
 
-FString FAssetImportParams::GetAssetTypePath(const FString& RootDestination, const FString& AssetType)
+FString FMSImportParams::GetAssetTypePath(const FString& RootDestination, const FString& AssetType)
 {
 	TMap<FString, FString> TypeNames = {
 		{TEXT("3d"),TEXT("3D_Assets")},
@@ -204,24 +204,24 @@ FString FAssetImportParams::GetAssetTypePath(const FString& RootDestination, con
 	return FPaths::Combine(RootDestination, TypeNames[AssetType]);
 }
 
-EAtlasSubType FAssetImportParams::GetAtlasSubtype(TSharedPtr<FAssetTypeData> AssetImportData)
+EAtlasSubType FMSImportParams::GetAtlasSubtype(TSharedPtr<FAssetTypeData> AssetImportData)
 {
 	auto& TagsList = AssetImportData->AssetMetaInfo->Tags;
 	return (!TagsList.Contains(TEXT("decal"))) ? EAtlasSubType::ATLAS: EAtlasSubType::DECAL;
 }
-ESurfaceSubType FAssetImportParams::GetSurfaceSubtype(TSharedPtr<FAssetTypeData> AssetImportData)
+ESurfaceSubType FMSImportParams::GetSurfaceSubtype(TSharedPtr<FAssetTypeData> AssetImportData)
 {
 	auto& TagsList = AssetImportData->AssetMetaInfo->Tags;
 	return (TagsList.Contains(TEXT("metal"))) ? ESurfaceSubType::METAL : (TagsList.Contains(TEXT("surface imperfection"))) ? ESurfaceSubType::IMPERFECTION : (TagsList.Contains(TEXT("tileable displacement"))) ? ESurfaceSubType::DISPLACEMENT : ESurfaceSubType::SURFACE;
 }
 
-EAsset3DSubtype FAssetImportParams::Get3DSubtype(TSharedPtr<FAssetTypeData> AssetImportData)
+EAsset3DSubtype FMSImportParams::Get3DSubtype(TSharedPtr<FAssetTypeData> AssetImportData)
 {
 	return (AssetImportData->AssetMetaInfo->Categories.Contains(TEXT("scatter")) || AssetImportData->AssetMetaInfo->Tags.Contains(TEXT("scatter")) || AssetImportData->AssetMetaInfo->Categories.Contains(TEXT("cmb_asset")) || AssetImportData->AssetMetaInfo->Tags.Contains(TEXT("cmb_asset"))) ? EAsset3DSubtype::MULTI_MESH : EAsset3DSubtype::NORMAL_3D;
 	
 }
 
-FString FAssetImportParams::GetMasterMaterial(const FString& SelectedMaterial)
+FString FMSImportParams::GetMasterMaterial(const FString& SelectedMaterial)
 {
 	CopyPresetTextures();
 	FString MaterialPath = GetMaterial(SelectedMaterial);
diff --git a/Engine/Plugins/MegascansPlugin/Source/MegascansPlugin/Private/Utilities/MSAssetData.h b/Engine/Plugins/MegascansPlugin/Source/MegascansPlugin/Private/Utilities/MSAssetData.h
--- a/Engine/Plugins/MegascansPlugin/Source/MegascansPlugin/Private/Utilities/MSAssetData.h
+++ b/Engine/Plugins/MegascansPlugin/Source/MegascansPlugin/Private/Utilities/MSAssetData.h
@@ -96,9 +96,9 @@ struct ImportParamsSurfaceAsset
 
 
 
-class FAssetImportParams {
+class FMSImportParams {
 private:
-	FAssetImportParams() = default;
+	FMSImportParams() = default;
 
 	const TMap<EAtlasSubType, FString> AtlasSubTypeMaterials = {
 		{EAtlasSubType::ATLAS,"Atlas_Material"},
@@ -129,7 +129,7 @@ private:
 
 	const FString PlantsMasterMaterial = TEXT("Foliage_Material");
 	
-	static TSharedPtr<FAssetImportParams> ImportParamsInst;
+	static TSharedPtr<FMSImportParams> ImportParamsInst;
 	FString GetAssetTypePath(const FString& RootDestination, const FString& AssetType);
 	EAtlasSubType GetAtlasSubtype(TSharedPtr<FAssetTypeData> AssetImportData);
 	ESurfaceSubType GetSurfaceSubtype(TSharedPtr<FAssetTypeData> AssetImportData);
@@ -139,7 +139,7 @@ private:
 	FString GetMasterMaterial(const FString& SelectedMaterial);
 
 public:
-	static TSharedPtr<FAssetImportParams> Get();	
+	static TSharedPtr<FMSImportParams> Get();	
 //Asset related data
 	TSharedPtr<ImportParams3DAsset> Get3DAssetsParams(TSharedPtr<FAssetTypeData> AssetImportData);
 	TSharedPtr<ImportParamsSurfaceAsset> GetSurfaceParams(TSharedPtr<FAssetTypeData> AssetImportData);
